[WebRemoteControl] WebSocket API for generating DisplayCluster previews

diff --git a/Engine/Plugins/Runtime/nDisplay/Source/DisplayCluster/Private/Game/EngineClasses/Scene/DisplayClusterRootActor.cpp b/Engine/Plugins/Runtime/nDisplay/Source/DisplayCluster/Private/Game/EngineClasses/Scene/DisplayClusterRootActor.cpp
--- a/Engine/Plugins/Runtime/nDisplay/Source/DisplayCluster/Private/Game/EngineClasses/Scene/DisplayClusterRootActor.cpp
+++ b/Engine/Plugins/Runtime/nDisplay/Source/DisplayCluster/Private/Game/EngineClasses/Scene/DisplayClusterRootActor.cpp
@@ -42,6 +42,8 @@
 
 #include "Game/EngineClasses/Scene/DisplayClusterRootActorInitializer.h"
 
+#include "Algo/MaxElement.h"
+
 #if WITH_EDITOR
 #include "IConcertSyncClientModule.h"
 #include "IConcertClientWorkspace.h"
@@ -1042,6 +1044,65 @@ UDisplayClusterCameraComponent* ADisplayClusterRootActor::GetDefaultCamera() con
 	return DefaultViewPoint;
 }
 
+USceneComponent* ADisplayClusterRootActor::GetCommonViewPoint() const
+{
+	if (UDisplayClusterConfigurationData* ConfigData = GetConfigData())
+	{
+		if (ConfigData->Cluster)
+		{
+			TMap<FString, int32> ViewPointCounts;
+
+			// Get all the camera names used by viewports in this cluster
+			for (const TPair<FString, UDisplayClusterConfigurationClusterNode*>& NodePair : ConfigData->Cluster->Nodes)
+			{
+				if (!NodePair.Value)
+				{
+					continue;
+				}
+
+				for (const TPair<FString, UDisplayClusterConfigurationViewport*>& ViewportPair : NodePair.Value->Viewports)
+				{
+					if (!ViewportPair.Value)
+					{
+						continue;
+					}
+
+					ViewPointCounts.FindOrAdd(ViewportPair.Value->Camera)++;
+				}
+			}
+
+			if (!ViewPointCounts.IsEmpty())
+			{
+				// Find the viewpoint with the most references
+				TPair<FString, int32>* MostCommonViewPointCount = Algo::MaxElementBy(ViewPointCounts,
+					[](const TPair<FString, int32>& Pair)
+					{
+						return Pair.Value;
+					}
+				);
+
+				if (MostCommonViewPointCount && !MostCommonViewPointCount->Key.IsEmpty())
+				{
+					// Try to return the camera with the most common name
+					if (USceneComponent* Camera = GetComponentByName<UDisplayClusterCameraComponent>(MostCommonViewPointCount->Key))
+					{
+						return Camera;
+					}
+				}
+			}
+		}
+	}
+
+	// We didn't find a common camera override (or it was empty), so fall back to the default camera
+	if (USceneComponent* DefaultCamera = GetDefaultCamera())
+	{
+		return DefaultCamera;
+	}
+
+	// No default camera, so fall back to the cluster root
+	return GetRootComponent();
+}
+
 bool ADisplayClusterRootActor::SetReplaceTextureFlagForAllViewports(bool bReplace)
 {
 	IDisplayCluster& Display = IDisplayCluster::Get();
diff --git a/Engine/Plugins/Runtime/nDisplay/Source/DisplayCluster/Public/DisplayClusterRootActor.h b/Engine/Plugins/Runtime/nDisplay/Source/DisplayCluster/Public/DisplayClusterRootActor.h
--- a/Engine/Plugins/Runtime/nDisplay/Source/DisplayCluster/Public/DisplayClusterRootActor.h
+++ b/Engine/Plugins/Runtime/nDisplay/Source/DisplayCluster/Public/DisplayClusterRootActor.h
@@ -122,6 +122,13 @@ public:
 	UFUNCTION(BlueprintGetter)
 	UDisplayClusterCameraComponent* GetDefaultCamera() const;
 
+	/**
+	 * Get the view origin most commonly used by viewports in this cluster.
+	 * If no viewports override the camera, this returns the default camera, or if there isn't one, the actor's root component.
+	 */
+	UFUNCTION(BlueprintCallable, Category = "NDisplay|Components")
+	USceneComponent* GetCommonViewPoint() const;
+
 	UFUNCTION(BlueprintCallable, Category = "NDisplay|Render")
 	bool SetReplaceTextureFlagForAllViewports(bool bReplace);
 
diff --git a/Engine/Plugins/Runtime/nDisplay/Source/DisplayClusterLightCardEditor/Private/Viewport/DisplayClusterLightCardEditorViewportClient.cpp b/Engine/Plugins/Runtime/nDisplay/Source/DisplayClusterLightCardEditor/Private/Viewport/DisplayClusterLightCardEditorViewportClient.cpp
--- a/Engine/Plugins/Runtime/nDisplay/Source/DisplayClusterLightCardEditor/Private/Viewport/DisplayClusterLightCardEditorViewportClient.cpp
+++ b/Engine/Plugins/Runtime/nDisplay/Source/DisplayClusterLightCardEditor/Private/Viewport/DisplayClusterLightCardEditorViewportClient.cpp
@@ -1864,15 +1864,7 @@ USceneComponent* FDisplayClusterLightCardEditorViewportClient::FindProjectionOri
 		return nullptr;
 	}
 
-	TArray<UDisplayClusterCameraComponent*> ViewOriginComponents;
-	InRootActor->GetComponents<UDisplayClusterCameraComponent>(ViewOriginComponents);
-
-	if (ViewOriginComponents.Num())
-	{
-		return ViewOriginComponents[0];
-	}
-
-	return InRootActor->GetRootComponent();
+	return InRootActor->GetCommonViewPoint();
 }
 
 bool FDisplayClusterLightCardEditorViewportClient::IsLightCardSelected(const AActor* Actor)
diff --git a/Engine/Plugins/Runtime/nDisplay/Source/DisplayClusterScenePreview/DisplayClusterScenePreview.Build.cs b/Engine/Plugins/Runtime/nDisplay/Source/DisplayClusterScenePreview/DisplayClusterScenePreview.Build.cs
--- a/Engine/Plugins/Runtime/nDisplay/Source/DisplayClusterScenePreview/DisplayClusterScenePreview.Build.cs
+++ b/Engine/Plugins/Runtime/nDisplay/Source/DisplayClusterScenePreview/DisplayClusterScenePreview.Build.cs
@@ -20,8 +20,7 @@ public class DisplayClusterScenePreview : ModuleRules
 				"Core",
 				"CoreUObject",
 				"Engine",
-				"RenderCore",
-				"UnrealEd"
+				"RenderCore"
 			});
 	}
 }
diff --git a/Engine/Plugins/Runtime/nDisplay/Source/DisplayClusterScenePreview/Private/DisplayClusterScenePreviewModule.cpp b/Engine/Plugins/Runtime/nDisplay/Source/DisplayClusterScenePreview/Private/DisplayClusterScenePreviewModule.cpp
--- a/Engine/Plugins/Runtime/nDisplay/Source/DisplayClusterScenePreview/Private/DisplayClusterScenePreviewModule.cpp
+++ b/Engine/Plugins/Runtime/nDisplay/Source/DisplayClusterScenePreview/Private/DisplayClusterScenePreviewModule.cpp
@@ -257,6 +257,7 @@ void FDisplayClusterScenePreviewModule::RegisterOrUnregisterGlobalActorEvents()
 		}
 	}
 
+#if WITH_EDITOR
 	if (bShouldBeRegistered && !bIsRegisteredForActorEvents)
 	{
 		// Register for events
@@ -279,10 +280,12 @@ void FDisplayClusterScenePreviewModule::RegisterOrUnregisterGlobalActorEvents()
 			GEngine->OnLevelActorDeleted().RemoveAll(this);
 		}
 	}
+#endif
 }
 
 void FDisplayClusterScenePreviewModule::RegisterRootActorEvents(AActor* Actor, bool bShouldRegister)
 {
+#if WITH_EDITOR
 	if (!Actor)
 	{
 		return;
@@ -299,6 +302,7 @@ void FDisplayClusterScenePreviewModule::RegisterRootActorEvents(AActor* Actor, b
 			Blueprint->OnCompiled().RemoveAll(this);
 		}
 	}
+#endif
 }
 
 void FDisplayClusterScenePreviewModule::AutoPopulateScene(FRendererConfig& RendererConfig)
@@ -407,6 +411,7 @@ void FDisplayClusterScenePreviewModule::OnLevelActorDeleted(AActor* Actor)
 
 void FDisplayClusterScenePreviewModule::OnBlueprintCompiled(UBlueprint* Blueprint)
 {
+#if WITH_EDITOR
 	for (TPair<int32, FRendererConfig>& ConfigPair : RendererConfigs)
 	{
 		FRendererConfig& Config = ConfigPair.Value;
@@ -415,6 +420,7 @@ void FDisplayClusterScenePreviewModule::OnBlueprintCompiled(UBlueprint* Blueprin
 			Config.bIsSceneDirty = true;
 		}
 	}
+#endif
 }
 
 void FDisplayClusterScenePreviewModule::OnObjectTransacted(UObject* Object, const FTransactionObjectEvent& TransactionObjectEvent)
diff --git a/Engine/Plugins/VirtualProduction/EpicStageApp/EpicStageApp.uplugin b/Engine/Plugins/VirtualProduction/EpicStageApp/EpicStageApp.uplugin
new file mode 100644
--- /dev/null
+++ b/Engine/Plugins/VirtualProduction/EpicStageApp/EpicStageApp.uplugin
@@ -0,0 +1,39 @@
+{
+	"FileVersion": 1,
+	"Version": 1,
+	"VersionName": "1.0",
+	"FriendlyName": "Epic Stage App",
+	"Description": "Enables remote connections from the Epic Stage App",
+	"Category": "Messaging",
+	"CreatedBy": "Epic Games Inc.",
+	"CreatedByURL": "http://epicgames.com",
+	"DocsURL": "",
+	"MarketplaceURL": "",
+	"SupportURL": "",
+	"CanContainContent": true,
+	"IsBetaVersion": true,
+	"IsExperimentalVersion": false,
+	"Installed": false,
+	"Modules": [
+		{
+			"Name": "EpicStageApp",
+			"Type": "Runtime",
+			"LoadingPhase": "Default",
+			"PlatformAllowList": [
+				"Mac",
+				"Win64",
+				"Linux"
+			]
+		}
+	],
+	"Plugins": [
+		{
+			"Name": "RemoteControl",
+			"Enabled": true
+		},
+		{
+			"Name": "nDisplay",
+			"Enabled": true
+		}
+	]
+}
\ No newline at end of file
diff --git a/Engine/Plugins/VirtualProduction/EpicStageApp/Source/EpicStageApp/EpicStageApp.build.cs b/Engine/Plugins/VirtualProduction/EpicStageApp/Source/EpicStageApp/EpicStageApp.build.cs
new file mode 100644
--- /dev/null
+++ b/Engine/Plugins/VirtualProduction/EpicStageApp/Source/EpicStageApp/EpicStageApp.build.cs
@@ -0,0 +1,33 @@
+// Copyright Epic Games, Inc. All Rights Reserved.
+
+using UnrealBuildTool;
+
+public class EpicStageApp : ModuleRules
+{
+	public EpicStageApp(ReadOnlyTargetRules Target) : base(Target)
+	{
+		PublicDependencyModuleNames.AddRange(
+			new string[] {
+				"Core",
+				"CoreUObject",
+				"Serialization"
+			}
+		);
+
+        PrivateDependencyModuleNames.AddRange(
+			new string[] {
+				"DeveloperSettings",
+				"DisplayCluster",
+				"DisplayClusterLightCardEditorShaders",
+				"DisplayClusterScenePreview",
+				"Engine",
+				"ImageWrapper",
+				"Networking",
+				"RemoteControl",
+				"RemoteControlCommon",
+				"Sockets",
+				"WebRemoteControl"
+			}
+        );
+	}
+}
diff --git a/Engine/Plugins/VirtualProduction/EpicStageApp/Source/EpicStageApp/Private/EpicStageApp.cpp b/Engine/Plugins/VirtualProduction/EpicStageApp/Source/EpicStageApp/Private/EpicStageApp.cpp
new file mode 100644
--- /dev/null
+++ b/Engine/Plugins/VirtualProduction/EpicStageApp/Source/EpicStageApp/Private/EpicStageApp.cpp
@@ -0,0 +1,21 @@
+#include "EpicStageApp.h"
+#include "Modules/ModuleManager.h"
+#include "IWebRemoteControlModule.h"
+
+void FEpicStageAppModule::StartupModule()
+{
+	IWebRemoteControlModule& WebRemoteControl = FModuleManager::LoadModuleChecked<IWebRemoteControlModule>("WebRemoteControl");
+	RouteHandler.RegisterRoutes(WebRemoteControl);
+
+	StageAppBeaconReceiver.Startup();
+}
+
+void FEpicStageAppModule::ShutdownModule()
+{
+	if (IWebRemoteControlModule* WebRemoteControl = FModuleManager::GetModulePtr<IWebRemoteControlModule>("WebRemoteControl"))
+	{
+		RouteHandler.UnregisterRoutes(*WebRemoteControl);
+	}
+
+	StageAppBeaconReceiver.Shutdown();
+}
diff --git a/Engine/Plugins/VirtualProduction/EpicStageApp/Source/EpicStageApp/Private/EpicStageApp.h b/Engine/Plugins/VirtualProduction/EpicStageApp/Source/EpicStageApp/Private/EpicStageApp.h
new file mode 100644
--- /dev/null
+++ b/Engine/Plugins/VirtualProduction/EpicStageApp/Source/EpicStageApp/Private/EpicStageApp.h
@@ -0,0 +1,23 @@
+#pragma once
+
+#include "CoreMinimal.h"
+
+#include "StageAppBeaconReceiver.h"
+#include "StageAppRouteHandler.h"
+#include "Modules/ModuleInterface.h"
+
+class FEpicStageAppModule : public IModuleInterface
+{
+	//~ Begin IModuleInterface interface
+	virtual void StartupModule() override;
+	virtual void ShutdownModule() override;
+	//~ End IModuleInterface interface
+	
+	/** Handler for WebSocket routes. */
+	FStageAppRouteHandler RouteHandler;
+
+	/** Receiver that responds to beacon messages from the stage app. */
+	FStageAppBeaconReceiver StageAppBeaconReceiver;
+};
+
+IMPLEMENT_MODULE(FEpicStageAppModule, EpicStageApp);
\ No newline at end of file
diff --git a/Engine/Plugins/VirtualProduction/EpicStageApp/Source/EpicStageApp/Private/StageAppBeaconReceiver.cpp b/Engine/Plugins/VirtualProduction/EpicStageApp/Source/EpicStageApp/Private/StageAppBeaconReceiver.cpp
new file mode 100644
--- /dev/null
+++ b/Engine/Plugins/VirtualProduction/EpicStageApp/Source/EpicStageApp/Private/StageAppBeaconReceiver.cpp
@@ -0,0 +1,205 @@
+// Copyright Epic Games, Inc. All Rights Reserved.
+
+#include "StageAppBeaconReceiver.h"
+
+#include "Common/UdpSocketBuilder.h"
+#include "IRemoteControlModule.h"
+#include "MessageEndpoint.h"
+#include "MessageEndpointBuilder.h"
+#include "Misc/App.h"
+#include "Misc/CommandLine.h"
+#include "RemoteControlSettings.h"
+#include "Serialization/ArrayReader.h"
+#include "Serialization/ArrayWriter.h"
+#include "StageAppSettings.h"
+
+namespace StageAppBeaconReceiverConstants
+{
+	const uint8 ProtocolVersion = 0;
+	const TArray<uint8> ProtocolIdentifier = { 'E', 'S', '@', 'p' };
+}
+
+FStageAppBeaconReceiver::FStageAppBeaconReceiver()
+	: Guid(FGuid::NewGuid())
+{
+}
+
+void FStageAppBeaconReceiver::Startup()
+{
+	const UStageAppSettings& Settings = *GetDefault<UStageAppSettings>();
+	FIPv4Address DiscoveryAddress;
+	if (!FIPv4Address::Parse(Settings.DiscoveryEndpoint, DiscoveryAddress))
+	{
+		UE_LOG(LogRemoteControl, Error, TEXT("Failed to parse Stage App discovery endpoint address \"%s\""), *Settings.DiscoveryEndpoint);
+		return;
+	}
+
+	WebsocketPort = GetDefault<URemoteControlSettings>()->RemoteControlWebSocketServerPort;
+
+	Socket = FUdpSocketBuilder(TEXT("StageAppBeaconResponder"))
+		.AsNonBlocking()
+		.AsReusable()
+		.BoundToPort(Settings.DiscoveryPort)
+		.JoinedToGroup(DiscoveryAddress) // joins default interface
+		.WithMulticastLoopback()
+		.WithReceiveBufferSize(256);
+
+	if (!Socket)
+	{
+		UE_LOG(LogRemoteControl, Warning, TEXT("StageAppBeaconResponder failed to create multicast socket"));
+		return;
+	}
+
+	// Listen for multicast packets on all interfaces.
+	// Note: If a new interface appears later on, the socket will need to join it in order to receive multicast packets from it.
+	{
+		TSharedRef<FInternetAddr> DiscoveryInternetAddr = FIPv4Endpoint(DiscoveryAddress, 0).ToInternetAddr(); // Note: Port will be ignored by JoinMulticastGroup
+
+		TArray<TSharedPtr<FInternetAddr>> LocapIps;
+		ISocketSubsystem::Get(PLATFORM_SOCKETSUBSYSTEM)->GetLocalAdapterAddresses(LocapIps);
+
+		for (TSharedPtr<FInternetAddr>& LocapIp : LocapIps)
+		{
+			if (!LocapIp.IsValid())
+			{
+				continue;
+			}
+			
+			const bool bJoinedGroup = Socket->JoinMulticastGroup(*DiscoveryInternetAddr, *LocapIp);
+
+			if (!bJoinedGroup)
+			{
+				UE_LOG(LogRemoteControl, Warning, 
+					TEXT("StageAppBeaconResponder failed to join multicast group '%s' on detected local interface '%s'"), 
+					*DiscoveryAddress.ToString(), *LocapIp->ToString(false));
+			}
+		}
+	}
+
+	Thread.Reset(FRunnableThread::Create(this, TEXT("StageAppBeaconResponderThread"), 0, TPri_Normal, FPlatformAffinity::GetPoolThreadMask()));
+}
+
+void FStageAppBeaconReceiver::Shutdown()
+{
+	if (Thread)
+	{
+		Thread->Kill();
+		Thread = nullptr;
+	}
+
+	if (Socket)
+	{
+		ISocketSubsystem::Get()->DestroySocket(Socket);
+	}
+}
+
+bool FStageAppBeaconReceiver::Init()
+{
+	return true;
+}
+
+uint32 FStageAppBeaconReceiver::Run()
+{
+	const UStageAppSettings& Settings = *GetDefault<UStageAppSettings>();
+	const FTimespan WaitTime = FTimespan::FromSeconds(Settings.DiscoverySocketWaitTime);
+
+	while (!bStopping)
+	{
+		if (!Socket)
+		{
+			return 0;
+		}
+
+		Socket->Wait(ESocketWaitConditions::WaitForRead, WaitTime);
+		ReceiveBeaconMessages();
+	}
+
+	return 0;
+}
+
+void FStageAppBeaconReceiver::Stop()
+{
+	bStopping = true;
+}
+
+void FStageAppBeaconReceiver::ReceiveBeaconMessages()
+{
+	if (!Socket)
+	{
+		return;
+	}
+
+	uint32 PendingDataSize;
+	if (Socket->HasPendingData(PendingDataSize) && PendingDataSize > 0)
+	{
+		FArrayReader MessageData = FArrayReader(true);
+		MessageData.SetNumUninitialized(PendingDataSize);
+
+		int32 NumRead;
+		TSharedRef<FInternetAddr> Source = ISocketSubsystem::Get()->CreateInternetAddr();
+		if (Socket->RecvFrom(MessageData.GetData(), MessageData.Num(), NumRead, *Source) && NumRead > 0)
+		{
+			HandleBeaconMessage(MessageData, Source);
+		}
+	}
+}
+
+void FStageAppBeaconReceiver::HandleBeaconMessage(FArrayReader& MessageData, TSharedRef<FInternetAddr> Source)
+{
+	if (MessageData.Num() != 5)
+	{
+		return;
+	}
+
+	// Check that the protocol identifier matches. If not, this message is probably unrelated.
+	for (const uint8 ExpectedByte : StageAppBeaconReceiverConstants::ProtocolIdentifier)
+	{
+		uint8 AppByte;
+		MessageData << AppByte;
+
+		if (AppByte != ExpectedByte)
+		{
+			return;
+		}
+	}
+
+	// We don't do anything with this yet, but could use it to ignore beacons from apps with incompatible protocols
+	uint8 AppProtocolVersion;
+	MessageData << AppProtocolVersion;
+
+	const FString FriendlyName = GetFriendlyName();
+
+	FArrayWriter Writer(true);
+
+#if !PLATFORM_LITTLE_ENDIAN
+	// App expects little-endian byte order
+	Writer.SetByteSwapping(true);
+#endif
+
+	{
+		Writer << (uint8&)StageAppBeaconReceiverConstants::ProtocolVersion;
+		Writer << Guid;
+		Writer << (uint32&)WebsocketPort;
+		Writer << (FString&)FriendlyName;
+	}
+
+	int32 Sent;
+	Socket->SendTo(Writer.GetData(), Writer.Num(), Sent, *Source);
+}
+
+FString FStageAppBeaconReceiver::GetFriendlyName() const
+{
+	FString FriendlyName;
+
+	if (FParse::Value(FCommandLine::Get(), TEXT("-StageFriendlyName="), FriendlyName))
+	{
+		return FriendlyName;
+	}
+
+	if (FParse::Value(FCommandLine::Get(), TEXT("-CONCERTDISPLAYNAME="), FriendlyName))
+	{
+		return FriendlyName;
+	}
+
+	return FApp::GetSessionOwner();
+}
diff --git a/Engine/Plugins/VirtualProduction/EpicStageApp/Source/EpicStageApp/Private/StageAppBeaconReceiver.h b/Engine/Plugins/VirtualProduction/EpicStageApp/Source/EpicStageApp/Private/StageAppBeaconReceiver.h
new file mode 100644
--- /dev/null
+++ b/Engine/Plugins/VirtualProduction/EpicStageApp/Source/EpicStageApp/Private/StageAppBeaconReceiver.h
@@ -0,0 +1,92 @@
+// Copyright Epic Games, Inc. All Rights Reserved.
+
+#pragma once
+
+#include "CoreMinimal.h"
+#include "HAL/Runnable.h"
+#include "HAL/RunnableThread.h"
+#include "StageAppBeaconReceiver.generated.h"
+
+class FArrayReader;
+class FInternetAddr;
+class FSocket;
+
+USTRUCT()
+struct FRemoteControl_StageAppBeaconMessage
+{
+	GENERATED_BODY()
+};
+
+USTRUCT()
+struct FRemoteControl_StageAppConnectionInfoMessage
+{
+	GENERATED_BODY()
+
+	/** Version of the stage app beacon protocol to use. */
+	uint8 Protocol;
+
+	/** ID of the sender. */
+	FGuid SenderId;
+
+	/** The websocket port the app should connect to. */
+	uint32 WebsocketPort;
+
+	/** The websocket address the app should connect to. */
+	FString WebsocketAddress;
+
+	/** The user-friendly name of the connection to show to the user. */
+	FString FriendlyName;
+};
+
+/**
+ * Receives beacon messages from the Unreal Stage App and replies with connection information.
+ * This allows the app to detect compatible Unreal instances on the local network and list them for the user.
+ */
+class FStageAppBeaconReceiver
+	: public FRunnable
+{
+public:
+	/** Creates and initializes a new beacon receiver. */
+	FStageAppBeaconReceiver();
+
+	/** Destroy the beacon receiver. */
+	virtual ~FStageAppBeaconReceiver() {};
+
+	/** Open a socket and start a thread listening for beacon messages. */
+	void Startup();
+
+	/** Close the socket and kill the listening thread. */
+	void Shutdown();
+
+protected:
+	//~ Begin FRunnable implementation
+	virtual bool Init() override;
+	virtual uint32 Run() override;
+	virtual void Stop() override;
+	//~ End FRunnable implementation
+
+private:
+	/** Called regularly to receive beacon messages. */
+	void ReceiveBeaconMessages();
+
+	/** Handle a message received on the beacon endpoint. */
+	void HandleBeaconMessage(FArrayReader& MessageData, TSharedRef<FInternetAddr> Source);
+
+	/** Get the name to report to apps searching for the engine. */
+	FString GetFriendlyName() const;
+
+	/** Identifier for this engine instance. Used by the app to differentiate instances even if they're on the same machine. */
+	FGuid Guid;
+
+	/** The websocket port to reply with. */
+	uint32 WebsocketPort;
+
+	/** Socket used to listen for and reply to beacon messages from the app. */
+	FSocket* Socket;
+
+	/** Flag indicating that the thread is stopping. */
+	TAtomic<bool> bStopping;
+
+	/** The thread to receive beacon messages on. */
+	TUniquePtr<FRunnableThread> Thread = nullptr;
+};
\ No newline at end of file
diff --git a/Engine/Plugins/VirtualProduction/EpicStageApp/Source/EpicStageApp/Private/StageAppRequest.h b/Engine/Plugins/VirtualProduction/EpicStageApp/Source/EpicStageApp/Private/StageAppRequest.h
new file mode 100644
--- /dev/null
+++ b/Engine/Plugins/VirtualProduction/EpicStageApp/Source/EpicStageApp/Private/StageAppRequest.h
@@ -0,0 +1,166 @@
+// Copyright Epic Games, Inc. All Rights Reserved.
+#pragma  once
+
+#include "CoreMinimal.h"
+#include "IDisplayClusterScenePreview.h"
+#include "RemoteControlWebsocketRoute.h"
+
+#include "StageAppRequest.generated.h"
+
+/** Type of preview render to perform, exposed as a UENUM. Corresponds to FDisplayClusterScenePreviewRenderSettings::ERenderType. */
+UENUM()
+enum class ERCWebSocketNDisplayPreviewRenderType : uint8
+{
+	Color,
+	Normals
+};
+
+/** Type of projection to use for a preview render, exposed as a UENUM. Corresponds to EDisplayClusterMeshProjectionType. */
+UENUM()
+enum class ERCWebSocketNDisplayPreviewRenderProjectionType : uint8
+{
+	Perspective,
+	Azimuthal
+};
+
+/**
+ * Holds a request made via websocket to create an nDisplay preview renderer.
+ */
+USTRUCT()
+struct FRCWebSocketNDisplayPreviewRendererCreateBody : public FRCRequest
+{
+	GENERATED_BODY()
+
+	FRCWebSocketNDisplayPreviewRendererCreateBody()
+	{
+		AddStructParameter(ParametersFieldLabel());
+	}
+
+	/**
+	 * Get the label for the property value struct.
+	 */
+	static FString ParametersFieldLabel() { return TEXT("Parameters"); }
+
+	/**
+	 * The path of the root DisplayCluster actor to preview. This may be empty, in which case you can set the root actor later.
+	 */
+	UPROPERTY()
+	FString RootActorPath = "";
+};
+
+/**
+ * Holds a request made via websocket to change the root DisplayCluster actor of an nDisplay preview renderer.
+ */
+USTRUCT()
+struct FRCWebSocketNDisplayPreviewRendererSetRootBody : public FRCRequest
+{
+	GENERATED_BODY()
+
+	FRCWebSocketNDisplayPreviewRendererSetRootBody()
+	{
+		AddStructParameter(ParametersFieldLabel());
+	}
+
+	/**
+	 * Get the label for the property value struct.
+	 */
+	static FString ParametersFieldLabel() { return TEXT("Parameters"); }
+
+	/**
+	 * The ID of the renderer returned when it was created.
+	 */
+	UPROPERTY()
+	int32 RendererId = -1;
+
+	/**
+	 * The path of the root DisplayCluster actor to preview.
+	 */
+	UPROPERTY()
+	FString RootActorPath = "";
+};
+
+/**
+ * Holds a request made via websocket to destroy an nDisplay preview renderer.
+ */
+USTRUCT()
+struct FRCWebSocketNDisplayPreviewRendererDestroyBody : public FRCRequest
+{
+	GENERATED_BODY()
+
+	FRCWebSocketNDisplayPreviewRendererDestroyBody()
+	{
+		AddStructParameter(ParametersFieldLabel());
+	}
+
+	/**
+	 * Get the label for the property value struct.
+	 */
+	static FString ParametersFieldLabel() { return TEXT("Parameters"); }
+
+	/**
+	 * The ID of the renderer returned when it was created.
+	 */
+	UPROPERTY()
+	int32 RendererId = -1;
+};
+
+/**
+ * Holds a request made via websocket to render a preview of a display cluster.
+ */
+USTRUCT()
+struct FRCWebSocketNDisplayPreviewRenderBody : public FRCRequest
+{
+	GENERATED_BODY()
+
+	FRCWebSocketNDisplayPreviewRenderBody()
+	{
+		AddStructParameter(ParametersFieldLabel());
+	}
+
+	/**
+	 * Get the label for the property value struct.
+	 */
+	static FString ParametersFieldLabel() { return TEXT("Parameters"); }
+
+	/**
+	 * The ID of the renderer returned when it was created.
+	 */
+	UPROPERTY()
+	int32 RendererId = -1;
+
+	/**
+	 * The type of render to perform.
+	 */
+	UPROPERTY()
+	ERCWebSocketNDisplayPreviewRenderType RenderType = ERCWebSocketNDisplayPreviewRenderType::Color;
+
+	/**
+	 * The type of projection to use.
+	 */
+	UPROPERTY()
+	ERCWebSocketNDisplayPreviewRenderProjectionType ProjectionType = ERCWebSocketNDisplayPreviewRenderProjectionType::Azimuthal;
+
+	/**
+	 * The resolution of the image to render.
+	 */
+	UPROPERTY()
+	FIntPoint Resolution = FIntPoint(1024, 1024);
+
+	/**
+	 * The preview camera's field of view (both horizontal and vertical) in degrees.
+	 */
+	UPROPERTY()
+	float FOV = 130.0f;
+
+	/**
+	 * The preview camera's Euler rotation relative to the camera's actual rotation in the scene.
+	 */
+	UPROPERTY()
+	FRotator Rotation = FRotator(90, 0, 0);
+
+	/**
+	 * The quality of the JPEG to send back to the requesting client, in range 50-100.
+	 */
+	UPROPERTY()
+	int32 JpegQuality = 50;
+};
diff --git a/Engine/Plugins/VirtualProduction/EpicStageApp/Source/EpicStageApp/Private/StageAppResponse.h b/Engine/Plugins/VirtualProduction/EpicStageApp/Source/EpicStageApp/Private/StageAppResponse.h
new file mode 100644
--- /dev/null
+++ b/Engine/Plugins/VirtualProduction/EpicStageApp/Source/EpicStageApp/Private/StageAppResponse.h
@@ -0,0 +1,65 @@
+// Copyright Epic Games, Inc. All Rights Reserved.
+#pragma once
+
+#include "CoreMinimal.h"
+
+#include "StageAppResponse.generated.h"
+
+/**
+ * Event sent to a client that created a renderer.
+ */
+USTRUCT()
+struct FRCPreviewRendererCreatedEvent
+{
+	GENERATED_BODY()
+
+	FRCPreviewRendererCreatedEvent()
+	: Type(TEXT("PreviewRendererCreated"))
+	{
+	}
+
+	/**
+	 * Type of the event.
+	 */
+	UPROPERTY()
+	FString Type;
+
+	/**
+	 * The ID of the new renderer
+	 */
+	UPROPERTY()
+	int32 RendererId = -1;
+};
+
+
+/**
+ * Event sent to a client that created a renderer.
+ */
+USTRUCT()
+struct FRCPreviewRenderCompletedEvent
+{
+	GENERATED_BODY()
+
+	FRCPreviewRenderCompletedEvent()
+	: Type(TEXT("PreviewRenderCompleted"))
+	{
+	}
+
+	/**
+	 * Type of the event.
+	 */
+	UPROPERTY()
+	FString Type;
+
+	/**
+	 * The ID of the renderer that created the image.
+	 */
+	UPROPERTY()
+	int32 RendererId = -1;
+
+	/**
+	 * Base64 encoded JPEG data of the preview render.
+	 */
+	UPROPERTY()
+	FString ImageBase64 = "";
+};
diff --git a/Engine/Plugins/VirtualProduction/EpicStageApp/Source/EpicStageApp/Private/StageAppRouteHandler.cpp b/Engine/Plugins/VirtualProduction/EpicStageApp/Source/EpicStageApp/Private/StageAppRouteHandler.cpp
new file mode 100644
--- /dev/null
+++ b/Engine/Plugins/VirtualProduction/EpicStageApp/Source/EpicStageApp/Private/StageAppRouteHandler.cpp
@@ -0,0 +1,280 @@
+#include "StageAppRouteHandler.h"
+
+#include "CanvasTypes.h"
+#include "Components/DisplayClusterCameraComponent.h"
+#include "DisplayClusterLightCardActor.h"
+#include "DisplayClusterRootActor.h"
+#include "Engine/TextureRenderTarget2D.h"
+#include "EpicStageApp.h"
+#include "IDisplayClusterScenePreview.h"
+#include "IImageWrapperModule.h"
+#include "Misc/Base64.h"
+#include "Modules/ModuleManager.h"
+#include "RemoteControlWebsocketRoute.h"
+#include "StageAppRequest.h"
+#include "StageAppResponse.h"
+#include "WebRemoteControlUtils.h"
+
+void FStageAppRouteHandler::RegisterRoutes(IWebRemoteControlModule& WebRemoteControl)
+{
+	RemoteControlModule = &WebRemoteControl;
+	RemoteControlModule->OnWebSocketConnectionClosed().AddRaw(this, &FStageAppRouteHandler::HandleClientDisconnected);
+
+	ImageWrapperModule = &FModuleManager::LoadModuleChecked<IImageWrapperModule>(FName("ImageWrapper"));
+
+	RegisterRoute(MakeUnique<FRemoteControlWebsocketRoute>(
+		TEXT("Create a preview renderer for a specific nDisplay cluster"),
+		TEXT("ndisplay.preview.renderer.create"),
+		FWebSocketMessageDelegate::CreateRaw(this, &FStageAppRouteHandler::HandleWebSocketNDisplayPreviewRendererCreate)
+	));
+
+	RegisterRoute(MakeUnique<FRemoteControlWebsocketRoute>(
+		TEXT("Set the root DisplayCluster actor for an nDisplay preview renderer"),
+		TEXT("ndisplay.preview.renderer.setroot"),
+		FWebSocketMessageDelegate::CreateRaw(this, &FStageAppRouteHandler::HandleWebSocketNDisplayPreviewRendererSetRoot)
+	));
+
+	RegisterRoute(MakeUnique<FRemoteControlWebsocketRoute>(
+		TEXT("Destroy an nDisplay preview renderer"),
+		TEXT("ndisplay.preview.renderer.destroy"),
+		FWebSocketMessageDelegate::CreateRaw(this, &FStageAppRouteHandler::HandleWebSocketNDisplayPreviewRendererDestroy)
+	));
+
+	RegisterRoute(MakeUnique<FRemoteControlWebsocketRoute>(
+		TEXT("Request a preview render from an nDisplay preview renderer"),
+		TEXT("ndisplay.preview.render"),
+		FWebSocketMessageDelegate::CreateRaw(this, &FStageAppRouteHandler::HandleWebSocketNDisplayPreviewRender)
+	));
+}
+
+void FStageAppRouteHandler::UnregisterRoutes(IWebRemoteControlModule& WebRemoteControl)
+{
+	if (RemoteControlModule == &WebRemoteControl)
+	{
+		for (const TUniquePtr<FRemoteControlWebsocketRoute>& Route : Routes)
+		{
+			WebRemoteControl.UnregisterWebsocketRoute(*Route);
+		}
+	}
+}
+
+void FStageAppRouteHandler::RegisterRoute(TUniquePtr<FRemoteControlWebsocketRoute> Route)
+{
+	checkSlow(RemoteControlModule);
+
+	RemoteControlModule->RegisterWebsocketRoute(*Route);
+	Routes.Emplace(MoveTemp(Route));
+}
+
+void FStageAppRouteHandler::HandleWebSocketNDisplayPreviewRendererCreate(const FRemoteControlWebSocketMessage& WebSocketMessage)
+{
+	FRCWebSocketNDisplayPreviewRendererCreateBody Body;
+	if (!WebRemoteControlUtils::DeserializeMessage(WebSocketMessage.RequestPayload, Body))
+	{
+		return;
+	}
+
+	const int32 RendererId = IDisplayClusterScenePreview::Get().CreateRenderer();
+	NDisplayPreviewRendererIdsByClientId.FindOrAdd(WebSocketMessage.ClientId).Add(RendererId);
+
+	if (!Body.RootActorPath.IsEmpty())
+	{
+		IDisplayClusterScenePreview::Get().SetRendererRootActorPath(RendererId, Body.RootActorPath, true);
+	}
+
+	FRCPreviewRendererCreatedEvent Event;
+	Event.RendererId = RendererId;
+
+	TArray<uint8> Payload;
+	WebRemoteControlUtils::SerializeMessage(Event, Payload);
+
+	RemoteControlModule->SendWebsocketMessage(WebSocketMessage.ClientId, Payload);
+}
+
+void FStageAppRouteHandler::HandleWebSocketNDisplayPreviewRendererSetRoot(const FRemoteControlWebSocketMessage& WebSocketMessage)
+{
+	FRCWebSocketNDisplayPreviewRendererSetRootBody Body;
+	if (!WebRemoteControlUtils::DeserializeMessage(WebSocketMessage.RequestPayload, Body))
+	{
+		return;
+	}
+
+	IDisplayClusterScenePreview::Get().SetRendererRootActorPath(Body.RendererId, Body.RootActorPath, true);
+}
+
+void FStageAppRouteHandler::HandleWebSocketNDisplayPreviewRendererDestroy(const FRemoteControlWebSocketMessage& WebSocketMessage)
+{
+	FRCWebSocketNDisplayPreviewRendererDestroyBody Body;
+	if (!WebRemoteControlUtils::DeserializeMessage(WebSocketMessage.RequestPayload, Body))
+	{
+		return;
+	}
+
+	if (TArray<int32>* ClientRendererIds = NDisplayPreviewRendererIdsByClientId.Find(WebSocketMessage.ClientId))
+	{
+		// Check that this client created the renderer
+		if (ClientRendererIds->Remove(Body.RendererId) > 0)
+		{
+			IDisplayClusterScenePreview::Get().DestroyRenderer(Body.RendererId);
+		}
+	}
+}
+
+void FStageAppRouteHandler::HandleWebSocketNDisplayPreviewRender(const FRemoteControlWebSocketMessage& WebSocketMessage)
+{
+	if (!ImageWrapperModule)
+	{
+		return;
+	}
+
+	FRCWebSocketNDisplayPreviewRenderBody Body;
+	if (!WebRemoteControlUtils::DeserializeMessage(WebSocketMessage.RequestPayload, Body))
+	{
+		return;
+	}
+
+	IDisplayClusterScenePreview& PreviewModule = IDisplayClusterScenePreview::Get();
+	ADisplayClusterRootActor* RootActor = PreviewModule.GetRendererRootActor(Body.RendererId);
+	if (RootActor == nullptr)
+	{
+		return;
+	}
+
+	UWorld* World = RootActor->GetWorld();
+	if (!World)
+	{
+		return;
+	}
+
+	USceneComponent* ViewOriginComponent = RootActor->GetCommonViewPoint();
+	if (!ViewOriginComponent)
+	{
+		return;
+	}
+
+	// Set up the render
+	FDisplayClusterScenePreviewRenderSettings RenderSettings;
+
+	switch (Body.RenderType)
+	{
+	case ERCWebSocketNDisplayPreviewRenderType::Color:
+		RenderSettings.RenderType = FDisplayClusterScenePreviewRenderSettings::ERenderType::Color;
+		break;
+
+	case ERCWebSocketNDisplayPreviewRenderType::Normals:
+		RenderSettings.RenderType = FDisplayClusterScenePreviewRenderSettings::ERenderType::Normals;
+		break;
+	}
+
+	switch (Body.ProjectionType)
+	{
+	case ERCWebSocketNDisplayPreviewRenderProjectionType::Azimuthal:
+		RenderSettings.ProjectionType = EDisplayClusterMeshProjectionType::Azimuthal;
+		break;
+
+	case ERCWebSocketNDisplayPreviewRenderProjectionType::Perspective:
+		RenderSettings.ProjectionType = EDisplayClusterMeshProjectionType::Perspective;
+		break;
+	}
+
+	RenderSettings.EngineShowFlags.SetSelectionOutline(false);	// Determine FOV multipliers to match render target's aspect ratio
+	float XAxisMultiplier;
+	float YAxisMultiplier;
+
+	const FIntPoint TargetSize = Body.Resolution;
+	if (TargetSize.X > TargetSize.Y)
+	{
+		XAxisMultiplier = 1.0f;
+		YAxisMultiplier = TargetSize.X / (float)TargetSize.Y;
+	}
+	else
+	{
+		XAxisMultiplier = TargetSize.Y / (float)TargetSize.X;
+		YAxisMultiplier = 1.0f;
+	}
+
+	const float FOV = Body.FOV * (float)PI / 360.0f;
+
+	FSceneViewInitOptions& ViewInitOptions = RenderSettings.ViewInitOptions;
+	ViewInitOptions.BackgroundColor = FLinearColor::Black;
+	ViewInitOptions.SetViewRectangle(FIntRect(0, 0, TargetSize.X, TargetSize.Y));
+	ViewInitOptions.ViewOrigin = ViewOriginComponent->GetComponentLocation();
+	ViewInitOptions.ViewLocation = ViewInitOptions.ViewOrigin;
+	ViewInitOptions.ViewRotation = ViewOriginComponent->GetComponentRotation() + Body.Rotation;
+
+	// Rotate view 90 degrees to match FDisplayClusterLightCardEditorViewportClient
+	ViewInitOptions.ViewRotationMatrix = FInverseRotationMatrix(ViewInitOptions.ViewRotation) * FMatrix(
+		FPlane(0, 0, 1, 0),
+		FPlane(1, 0, 0, 0),
+		FPlane(0, 1, 0, 0),
+		FPlane(0, 0, 0, 1)
+	);
+
+	if ((bool)ERHIZBuffer::IsInverted)
+	{
+		ViewInitOptions.ProjectionMatrix = FReversedZPerspectiveMatrix(
+			FOV,
+			FOV,
+			XAxisMultiplier,
+			YAxisMultiplier,
+			GNearClippingPlane,
+			GNearClippingPlane
+		);
+	}
+	else
+	{
+		ViewInitOptions.ProjectionMatrix = FPerspectiveMatrix(
+			FOV,
+			FOV,
+			XAxisMultiplier,
+			YAxisMultiplier,
+			GNearClippingPlane,
+			GNearClippingPlane
+		);
+	}
+
+	const FIntPoint Resolution = Body.Resolution;
+	const FGuid ClientId = WebSocketMessage.ClientId;
+	const int32 RendererId = Body.RendererId;
+	const int32 JpegQuality = Body.JpegQuality;
+
+	// Render the image
+	PreviewModule.RenderQueued(Body.RendererId, RenderSettings, Resolution, FRenderResultDelegate::CreateLambda(
+		[this, ClientId, RendererId, Resolution, JpegQuality](FRenderTarget& RenderTarget)
+		{
+			// Save the image
+			TArray<FColor> PixelData;
+			RenderTarget.ReadPixels(PixelData);
+			FImageView ImageView(PixelData.GetData(), Resolution.X, Resolution.Y, ERawImageFormat::BGRA8);
+
+			// Convert to JPEG
+			TArray64<uint8> JpegData;
+			ImageWrapperModule->CompressImage(JpegData, EImageFormat::JPEG, ImageView, JpegQuality);
+
+			// Encode to Base64
+			FRCPreviewRenderCompletedEvent Event;
+			Event.RendererId = RendererId;
+			Event.ImageBase64 = FBase64::Encode(JpegData.GetData(), JpegData.Num());
+
+			// Send it back to the client
+			TArray<uint8> Payload;
+			WebRemoteControlUtils::SerializeMessage(Event, Payload);
+			RemoteControlModule->SendWebsocketMessage(ClientId, Payload);
+		}
+	));
+}
+
+void FStageAppRouteHandler::HandleClientDisconnected(FGuid ClientId)
+{
+	/** Destroy the client's renderers. */
+	if (const TArray<int32>* RendererIds = NDisplayPreviewRendererIdsByClientId.Find(ClientId))
+	{
+		IDisplayClusterScenePreview& PreviewModule = IDisplayClusterScenePreview::Get();
+		for (const int32 RendererId : *RendererIds)
+		{
+			PreviewModule.DestroyRenderer(RendererId);
+		}
+		NDisplayPreviewRendererIdsByClientId.Remove(ClientId);
+	}
+
+}
diff --git a/Engine/Plugins/VirtualProduction/EpicStageApp/Source/EpicStageApp/Private/StageAppRouteHandler.h b/Engine/Plugins/VirtualProduction/EpicStageApp/Source/EpicStageApp/Private/StageAppRouteHandler.h
new file mode 100644
--- /dev/null
+++ b/Engine/Plugins/VirtualProduction/EpicStageApp/Source/EpicStageApp/Private/StageAppRouteHandler.h
@@ -0,0 +1,48 @@
+#pragma once
+
+#include "CoreMinimal.h"
+#include "UObject/ObjectMacros.h"
+#include "IWebRemoteControlModule.h"
+#include "RemoteControlWebsocketRoute.h"
+
+class FStageAppRouteHandler
+{
+public:
+	/** Registers the WebSocket routes with the provided WebRemoteControl module. */
+	void RegisterRoutes(IWebRemoteControlModule& WebRemoteControl);
+
+	/** Unregisters the WebSocket routes with the provided WebRemoteControl module. */
+	void UnregisterRoutes(IWebRemoteControlModule& WebRemoteControl);
+
+private:
+	/** Register a WebSocket route with the WebRemoteControl module. */
+	void RegisterRoute(TUniquePtr<FRemoteControlWebsocketRoute> Route);
+
+	/** Handles creating a preview renderer. */
+	void HandleWebSocketNDisplayPreviewRendererCreate(const FRemoteControlWebSocketMessage& WebSocketMessage);
+
+	/** Handles setting the root actor of a renderer. */
+	void HandleWebSocketNDisplayPreviewRendererSetRoot(const FRemoteControlWebSocketMessage& WebSocketMessage);
+
+	/** Handles destroying a preview renderer. */
+	void HandleWebSocketNDisplayPreviewRendererDestroy(const FRemoteControlWebSocketMessage& WebSocketMessage);
+
+	/** Handles rendering a preview. */
+	void HandleWebSocketNDisplayPreviewRender(const FRemoteControlWebSocketMessage& WebSocketMessage);
+
+	/** Called when a client disconnects from the WebSocket server. */
+	void HandleClientDisconnected(FGuid ClientId);
+
+private:
+	/** The module for the WebSocket server for which this is handling routes. */
+	IWebRemoteControlModule* RemoteControlModule = nullptr;
+
+	/** Routes this has registered with the WebSocket server. */
+	TArray<TUniquePtr<FRemoteControlWebsocketRoute>> Routes;
+
+	/** nDisplay preview renderers belonging to each client. */
+	TMap<FGuid, TArray<int32>> NDisplayPreviewRendererIdsByClientId;
+
+	/** Image wrapper module used to compress preview images. */
+	class IImageWrapperModule* ImageWrapperModule = nullptr;
+};
diff --git a/Engine/Plugins/VirtualProduction/EpicStageApp/Source/EpicStageApp/Public/StageAppSettings.h b/Engine/Plugins/VirtualProduction/EpicStageApp/Source/EpicStageApp/Public/StageAppSettings.h
new file mode 100644
--- /dev/null
+++ b/Engine/Plugins/VirtualProduction/EpicStageApp/Source/EpicStageApp/Public/StageAppSettings.h
@@ -0,0 +1,57 @@
+// Copyright Epic Games, Inc. All Rights Reserved.
+
+#pragma once
+
+#include "CoreMinimal.h"
+#include "Engine/DeveloperSettings.h"
+#include "StageAppSettings.generated.h"
+
+/**
+ * Global settings for Epic Stage App integration.
+ */
+UCLASS(config = Engine)
+class EPICSTAGEAPP_API UStageAppSettings : public UDeveloperSettings
+{
+	GENERATED_BODY()
+
+public:
+	//~ Begin UDeveloperSettings implementation
+	virtual FName GetContainerName() const
+	{
+		return "Project";
+	}
+
+	virtual FName GetCategoryName() const
+	{
+		return "Plugins";
+	}
+
+	virtual FName GetSectionName() const
+	{
+		return "Epic Stage App";
+	}
+
+	virtual FText GetSectionText() const
+	{
+		return NSLOCTEXT("StageAppSettings", "StageAppSettingsSection", "Epic Stage App");
+	}
+	//~ End UDeveloperSettings implementation
+
+	/**
+	 * The IP endpoint to listen to for app discovery messages.
+	 */
+	UPROPERTY(config, EditAnywhere, Category = "Epic Stage App")
+	FString DiscoveryEndpoint = "230.0.0.2";
+
+	/**
+	 * The port to listen to for app discovery messages.
+	 */
+	UPROPERTY(config, EditAnywhere, Category = "Epic Stage App")
+	int32 DiscoveryPort = 6667;
+
+	/**
+	 * How many seconds to wait for app discovery messages.
+	 */
+	UPROPERTY(config, EditAnywhere, Category = "Epic Stage App")
+	float DiscoverySocketWaitTime = 10.f;
+};
diff --git a/Engine/Plugins/VirtualProduction/RemoteControl/Source/WebRemoteControl/Private/RemoteControlRequest.h b/Engine/Plugins/VirtualProduction/RemoteControl/Source/WebRemoteControl/Private/RemoteControlRequest.h
--- a/Engine/Plugins/VirtualProduction/RemoteControl/Source/WebRemoteControl/Private/RemoteControlRequest.h
+++ b/Engine/Plugins/VirtualProduction/RemoteControl/Source/WebRemoteControl/Private/RemoteControlRequest.h
@@ -4,52 +4,9 @@
 #include "CoreMinimal.h"
 #include "IRemoteControlModule.h"
 #include "RemoteControlModels.h"
-#include "RemoteControlRequest.generated.h"
-
-struct FBlockDelimiters
-{
-	int64 BlockStart = -1;
-	int64 BlockEnd = -1;
-
-	/** Get the size of current block */
-	int64 GetBlockSize() const { return BlockEnd - BlockStart; }
-};
-
-/**
- * Holds a request made to the remote control server.
- */
-USTRUCT()
-struct FRCRequest
-{
-	GENERATED_BODY()
-
-	virtual ~FRCRequest() = default;
-
-	TMap<FString, FBlockDelimiters>& GetStructParameters()
-	{
-		return StructParameters;
-	}
-
-	FBlockDelimiters& GetParameterDelimiters(const FString& ParameterName)
-	{
-		return StructParameters.FindChecked(ParameterName);
-	}
-
-	UPROPERTY()
-	FString Passphrase;
+#include "RemoteControlWebsocketRoute.h"
 
-	/** Holds the request's TCHAR payload. */
-	TArray<uint8> TCHARBody;
-
-protected:
-	void AddStructParameter(FString ParameterName)
-	{
-		StructParameters.Add(MoveTemp(ParameterName), FBlockDelimiters());
-	}
-
-	/** Holds the start and end of struct parameters */
-	TMap<FString, FBlockDelimiters> StructParameters;
-};
+#include "RemoteControlRequest.generated.h"
 
 USTRUCT()
 struct FRCRequestWrapper : public FRCRequest
diff --git a/Engine/Plugins/VirtualProduction/RemoteControl/Source/WebRemoteControl/Private/StageAppBeaconReceiver.cpp b/Engine/Plugins/VirtualProduction/RemoteControl/Source/WebRemoteControl/Private/StageAppBeaconReceiver.cpp
deleted file mode 100644
diff --git a/Engine/Plugins/VirtualProduction/RemoteControl/Source/WebRemoteControl/Private/StageAppBeaconReceiver.h b/Engine/Plugins/VirtualProduction/RemoteControl/Source/WebRemoteControl/Private/StageAppBeaconReceiver.h
deleted file mode 100644
diff --git a/Engine/Plugins/VirtualProduction/RemoteControl/Source/WebRemoteControl/Private/WebRemoteControl.cpp b/Engine/Plugins/VirtualProduction/RemoteControl/Source/WebRemoteControl/Private/WebRemoteControl.cpp
--- a/Engine/Plugins/VirtualProduction/RemoteControl/Source/WebRemoteControl/Private/WebRemoteControl.cpp
+++ b/Engine/Plugins/VirtualProduction/RemoteControl/Source/WebRemoteControl/Private/WebRemoteControl.cpp
@@ -305,8 +305,6 @@ void FWebRemoteControlModule::StartupModule()
 	{
 		StartWebSocketServer();
 	}
-
-	StageAppBeaconReceiver.Startup();
 }
 
 void FWebRemoteControlModule::ShutdownModule()
@@ -322,7 +320,6 @@ void FWebRemoteControlModule::ShutdownModule()
 		return;
 	}
 
-	StageAppBeaconReceiver.Shutdown();
 	EditorRoutes.UnregisterRoutes(this);
 	WebSocketHandler->UnregisterRoutes(this);
 	StopHttpServer();
diff --git a/Engine/Plugins/VirtualProduction/RemoteControl/Source/WebRemoteControl/Private/WebRemoteControl.h b/Engine/Plugins/VirtualProduction/RemoteControl/Source/WebRemoteControl/Private/WebRemoteControl.h
--- a/Engine/Plugins/VirtualProduction/RemoteControl/Source/WebRemoteControl/Private/WebRemoteControl.h
+++ b/Engine/Plugins/VirtualProduction/RemoteControl/Source/WebRemoteControl/Private/WebRemoteControl.h
@@ -13,7 +13,6 @@
 #include "RemoteControlRoute.h"
 #include "RemoteControlWebsocketRoute.h"
 #include "RemoteControlWebSocketServer.h"
-#include "StageAppBeaconReceiver.h"
 #include "WebRemoteControlEditorRoutes.h"
 
 struct FHttpServerRequest;
@@ -203,9 +202,6 @@ private:
 	FOnWebServerStarted OnWebSocketServerStartedDelegate;
 	FSimpleMulticastDelegate OnWebSocketServerStoppedDelegate;
 
-	/** Receiver that responds to beacon messages from the stage app. */
-	FStageAppBeaconReceiver StageAppBeaconReceiver;
-
 	/** Name cache for Get Controller Result structs
 	* Key: Controller name; Value: Name of the dynamic struct holding our result*/
 	TMap<FName, FString> ControllersSerializerStructNameCache;
diff --git a/Engine/Plugins/VirtualProduction/RemoteControl/Source/WebRemoteControl/Public/RemoteControlWebsocketRoute.h b/Engine/Plugins/VirtualProduction/RemoteControl/Source/WebRemoteControl/Public/RemoteControlWebsocketRoute.h
--- a/Engine/Plugins/VirtualProduction/RemoteControl/Source/WebRemoteControl/Public/RemoteControlWebsocketRoute.h
+++ b/Engine/Plugins/VirtualProduction/RemoteControl/Source/WebRemoteControl/Public/RemoteControlWebsocketRoute.h
@@ -2,6 +2,8 @@
 
 #pragma once
 
+#include "RemoteControlWebsocketRoute.generated.h"
+
 struct FRemoteControlWebSocketMessage
 {
 	FString MessageName;
@@ -31,3 +33,48 @@ struct FRemoteControlWebsocketRoute
 	friend uint32 GetTypeHash(const FRemoteControlWebsocketRoute& Route) { return GetTypeHash(Route.MessageName); }
 	friend bool operator==(const FRemoteControlWebsocketRoute& LHS, const FRemoteControlWebsocketRoute& RHS) { return LHS.MessageName == RHS.MessageName; }
 };
+
+struct FBlockDelimiters
+{
+	int64 BlockStart = -1;
+	int64 BlockEnd = -1;
+
+	/** Get the size of current block */
+	int64 GetBlockSize() const { return BlockEnd - BlockStart; }
+};
+
+/**
+ * Holds a request made to the remote control server.
+ */
+USTRUCT()
+struct FRCRequest
+{
+	GENERATED_BODY()
+
+	virtual ~FRCRequest() = default;
+
+	TMap<FString, FBlockDelimiters>& GetStructParameters()
+	{
+		return StructParameters;
+	}
+
+	FBlockDelimiters& GetParameterDelimiters(const FString& ParameterName)
+	{
+		return StructParameters.FindChecked(ParameterName);
+	}
+
+	UPROPERTY()
+	FString Passphrase;
+
+	/** Holds the request's TCHAR payload. */
+	TArray<uint8> TCHARBody;
+
+protected:
+	void AddStructParameter(FString ParameterName)
+	{
+		StructParameters.Add(MoveTemp(ParameterName), FBlockDelimiters());
+	}
+
+	/** Holds the start and end of struct parameters */
+	TMap<FString, FBlockDelimiters> StructParameters;
+};
diff --git a/Engine/Plugins/VirtualProduction/RemoteControl/Source/WebRemoteControl/Public/StageAppSettings.h b/Engine/Plugins/VirtualProduction/RemoteControl/Source/WebRemoteControl/Public/StageAppSettings.h
deleted file mode 100644
