Quick fix for compiler error in EpicStageApp after most recent nDisplay changes

diff --git a/Engine/Plugins/VirtualProduction/EpicStageApp/Source/EpicStageApp/Private/StageAppRouteHandler.cpp b/Engine/Plugins/VirtualProduction/EpicStageApp/Source/EpicStageApp/Private/StageAppRouteHandler.cpp
--- a/Engine/Plugins/VirtualProduction/EpicStageApp/Source/EpicStageApp/Private/StageAppRouteHandler.cpp
+++ b/Engine/Plugins/VirtualProduction/EpicStageApp/Source/EpicStageApp/Private/StageAppRouteHandler.cpp
@@ -173,7 +173,7 @@ void FStageAppRouteHandler::HandleWebSocketNDisplayPreviewRender(const FRemoteCo
 		break;
 
 	case ERCWebSocketNDisplayPreviewRenderProjectionType::Perspective:
-		RenderSettings.ProjectionType = EDisplayClusterMeshProjectionType::Perspective;
+		RenderSettings.ProjectionType = EDisplayClusterMeshProjectionType::Linear;
 		break;
 	}
 
