[EpicStageApp] Add missing copyright notices and fix header includes
Disable plugin on Mac (as nDisplay doesn't support Mac)

diff --git a/Engine/Plugins/VirtualProduction/EpicStageApp/EpicStageApp.uplugin b/Engine/Plugins/VirtualProduction/EpicStageApp/EpicStageApp.uplugin
--- a/Engine/Plugins/VirtualProduction/EpicStageApp/EpicStageApp.uplugin
+++ b/Engine/Plugins/VirtualProduction/EpicStageApp/EpicStageApp.uplugin
@@ -20,7 +20,6 @@
 			"Type": "Runtime",
 			"LoadingPhase": "Default",
 			"PlatformAllowList": [
-				"Mac",
 				"Win64",
 				"Linux"
 			]
diff --git a/Engine/Plugins/VirtualProduction/EpicStageApp/Source/EpicStageApp/Private/EpicStageApp.cpp b/Engine/Plugins/VirtualProduction/EpicStageApp/Source/EpicStageApp/Private/EpicStageApp.cpp
--- a/Engine/Plugins/VirtualProduction/EpicStageApp/Source/EpicStageApp/Private/EpicStageApp.cpp
+++ b/Engine/Plugins/VirtualProduction/EpicStageApp/Source/EpicStageApp/Private/EpicStageApp.cpp
@@ -1,5 +1,6 @@
+// Copyright Epic Games, Inc. All Rights Reserved.
+
 #include "EpicStageApp.h"
-#include "Modules/ModuleManager.h"
 #include "IWebRemoteControlModule.h"
 
 void FEpicStageAppModule::StartupModule()
diff --git a/Engine/Plugins/VirtualProduction/EpicStageApp/Source/EpicStageApp/Private/EpicStageApp.h b/Engine/Plugins/VirtualProduction/EpicStageApp/Source/EpicStageApp/Private/EpicStageApp.h
--- a/Engine/Plugins/VirtualProduction/EpicStageApp/Source/EpicStageApp/Private/EpicStageApp.h
+++ b/Engine/Plugins/VirtualProduction/EpicStageApp/Source/EpicStageApp/Private/EpicStageApp.h
@@ -1,3 +1,5 @@
+// Copyright Epic Games, Inc. All Rights Reserved.
+
 #pragma once
 
 #include "CoreMinimal.h"
@@ -5,6 +7,7 @@
 #include "StageAppBeaconReceiver.h"
 #include "StageAppRouteHandler.h"
 #include "Modules/ModuleInterface.h"
+#include "Modules/ModuleManager.h"
 
 class FEpicStageAppModule : public IModuleInterface
 {
diff --git a/Engine/Plugins/VirtualProduction/EpicStageApp/Source/EpicStageApp/Private/StageAppRouteHandler.cpp b/Engine/Plugins/VirtualProduction/EpicStageApp/Source/EpicStageApp/Private/StageAppRouteHandler.cpp
--- a/Engine/Plugins/VirtualProduction/EpicStageApp/Source/EpicStageApp/Private/StageAppRouteHandler.cpp
+++ b/Engine/Plugins/VirtualProduction/EpicStageApp/Source/EpicStageApp/Private/StageAppRouteHandler.cpp
@@ -1,3 +1,5 @@
+// Copyright Epic Games, Inc. All Rights Reserved.
+
 #include "StageAppRouteHandler.h"
 
 #include "CanvasTypes.h"
@@ -5,7 +7,6 @@
 #include "DisplayClusterLightCardActor.h"
 #include "DisplayClusterRootActor.h"
 #include "Engine/TextureRenderTarget2D.h"
-#include "EpicStageApp.h"
 #include "IDisplayClusterScenePreview.h"
 #include "IImageWrapperModule.h"
 #include "Misc/Base64.h"
diff --git a/Engine/Plugins/VirtualProduction/EpicStageApp/Source/EpicStageApp/Private/StageAppRouteHandler.h b/Engine/Plugins/VirtualProduction/EpicStageApp/Source/EpicStageApp/Private/StageAppRouteHandler.h
--- a/Engine/Plugins/VirtualProduction/EpicStageApp/Source/EpicStageApp/Private/StageAppRouteHandler.h
+++ b/Engine/Plugins/VirtualProduction/EpicStageApp/Source/EpicStageApp/Private/StageAppRouteHandler.h
@@ -1,3 +1,5 @@
+// Copyright Epic Games, Inc. All Rights Reserved.
+
 #pragma once
 
 #include "CoreMinimal.h"
