[EpicStageApp] Don't start beacon receiver when running a commandlet

diff --git a/Engine/Plugins/VirtualProduction/EpicStageApp/Source/EpicStageApp/Private/EpicStageApp.cpp b/Engine/Plugins/VirtualProduction/EpicStageApp/Source/EpicStageApp/Private/EpicStageApp.cpp
--- a/Engine/Plugins/VirtualProduction/EpicStageApp/Source/EpicStageApp/Private/EpicStageApp.cpp
+++ b/Engine/Plugins/VirtualProduction/EpicStageApp/Source/EpicStageApp/Private/EpicStageApp.cpp
@@ -8,7 +8,10 @@ void FEpicStageAppModule::StartupModule()
 	IWebRemoteControlModule& WebRemoteControl = FModuleManager::LoadModuleChecked<IWebRemoteControlModule>("WebRemoteControl");
 	RouteHandler.RegisterRoutes(WebRemoteControl);
 
-	StageAppBeaconReceiver.Startup();
+	if (!IsRunningCommandlet())
+	{
+		StageAppBeaconReceiver.Startup();
+	}
 }
 
 void FEpicStageAppModule::ShutdownModule()
@@ -18,5 +21,8 @@ void FEpicStageAppModule::ShutdownModule()
 		RouteHandler.UnregisterRoutes(*WebRemoteControl);
 	}
 
-	StageAppBeaconReceiver.Shutdown();
+	if (!IsRunningCommandlet())
+	{
+		StageAppBeaconReceiver.Shutdown();
+	}
 }
diff --git a/Engine/Plugins/VirtualProduction/EpicStageApp/Source/EpicStageApp/Private/StageAppBeaconReceiver.h b/Engine/Plugins/VirtualProduction/EpicStageApp/Source/EpicStageApp/Private/StageAppBeaconReceiver.h
--- a/Engine/Plugins/VirtualProduction/EpicStageApp/Source/EpicStageApp/Private/StageAppBeaconReceiver.h
+++ b/Engine/Plugins/VirtualProduction/EpicStageApp/Source/EpicStageApp/Private/StageAppBeaconReceiver.h
@@ -82,7 +82,7 @@ private:
 	uint32 WebsocketPort;
 
 	/** Socket used to listen for and reply to beacon messages from the app. */
-	FSocket* Socket;
+	FSocket* Socket = nullptr;
 
 	/** Flag indicating that the thread is stopping. */
 	TAtomic<bool> bStopping;
