[EpicStageApp][WebRemoteControl] Fix cases where a transaction could be cancelled but fail to report to subscribers

diff --git a/Engine/Plugins/VirtualProduction/EpicStageApp/Source/EpicStageApp/Private/StageAppRouteHandler.cpp b/Engine/Plugins/VirtualProduction/EpicStageApp/Source/EpicStageApp/Private/StageAppRouteHandler.cpp
--- a/Engine/Plugins/VirtualProduction/EpicStageApp/Source/EpicStageApp/Private/StageAppRouteHandler.cpp
+++ b/Engine/Plugins/VirtualProduction/EpicStageApp/Source/EpicStageApp/Private/StageAppRouteHandler.cpp
@@ -18,6 +18,7 @@
 
 #if WITH_EDITOR
 #include "Editor.h"
+#include "Editor/TransBuffer.h"
 #endif
 
 #define DEBUG_STAGEAPP_NORMALS false
@@ -180,6 +181,15 @@ void FStageAppRouteHandler::RegisterRoutes(IWebRemoteControlModule& WebRemoteCon
 
 	ImageWrapperModule = &FModuleManager::LoadModuleChecked<IImageWrapperModule>(FName("ImageWrapper"));
 
+	if (GEngine)
+	{
+		RegisterEngineEvents();
+	}
+	else
+	{
+		FCoreDelegates::OnPostEngineInit.AddRaw(this, &FStageAppRouteHandler::RegisterEngineEvents);
+	}
+
 	RegisterRoute(MakeUnique<FRemoteControlWebsocketRoute>(
 		TEXT("Create a preview renderer for a specific nDisplay cluster"),
 		TEXT("ndisplay.preview.renderer.create"),
@@ -223,6 +233,19 @@ void FStageAppRouteHandler::RegisterRoutes(IWebRemoteControlModule& WebRemoteCon
 	));
 }
 
+void FStageAppRouteHandler::RegisterEngineEvents()
+{
+#if WITH_EDITOR
+	if (GEditor)
+	{
+		if (UTransBuffer* TransBuffer = Cast<UTransBuffer>(GEditor->Trans))
+		{
+			TransBuffer->OnTransactionStateChanged().AddRaw(this, &FStageAppRouteHandler::OnTransactionStateChanged);
+		}
+	}
+#endif
+}
+
 void FStageAppRouteHandler::UnregisterRoutes(IWebRemoteControlModule& WebRemoteControl)
 {
 	if (RemoteControlModule == &WebRemoteControl)
@@ -238,6 +261,21 @@ void FStageAppRouteHandler::UnregisterRoutes(IWebRemoteControlModule& WebRemoteC
 		FTSTicker::GetCoreTicker().RemoveTicker(DragTimeoutTickerHandle);
 		DragTimeoutTickerHandle.Reset();
 	}
+
+	UnregisterEngineEvents();
+}
+
+void FStageAppRouteHandler::UnregisterEngineEvents()
+{
+#if WITH_EDITOR
+	if (GEditor)
+	{
+		if (UTransBuffer* TransBuffer = Cast<UTransBuffer>(GEditor->Trans))
+		{
+			TransBuffer->OnTransactionStateChanged().RemoveAll(this);
+		}
+	}
+#endif
 }
 
 void FStageAppRouteHandler::RegisterRoute(TUniquePtr<FRemoteControlWebsocketRoute> Route)
@@ -487,7 +525,18 @@ void FStageAppRouteHandler::HandleWebSocketNDisplayPreviewLightCardDragBegin(con
 	}
 
 #if WITH_EDITOR
-	GEditor->BeginTransaction(LOCTEXT("DragLightcardsTransaction", "Drag Light Cards with Stage App"));
+	PerRendererData->TransactionId.Invalidate();
+	if (GEditor && GEditor->Trans)
+	{
+		if (GEditor->BeginTransaction(LOCTEXT("DragLightcardsTransaction", "Drag Light Cards with Stage App")) != INDEX_NONE)
+		{
+			const FTransaction* Transaction = GEditor->Trans->GetTransaction(GEditor->Trans->GetQueueLength() - 1);
+			if (ensure(Transaction))
+			{
+				PerRendererData->TransactionId = Transaction->GetId();
+			}
+		}
+	}
 #endif
 
 	PerRendererData->UpdateDragSequenceNumber(Body.SequenceNumber);
@@ -515,6 +564,30 @@ void FStageAppRouteHandler::HandleWebSocketNDisplayPreviewLightCardDragBegin(con
 	}
 }
 
+#if WITH_EDITOR
+void FStageAppRouteHandler::OnTransactionStateChanged(const FTransactionContext& InTransactionContext, const ETransactionStateEventType InTransactionState)
+{
+	if (InTransactionState != ETransactionStateEventType::TransactionCanceled && InTransactionState != ETransactionStateEventType::TransactionFinalized)
+	{
+		return;
+	}
+
+	for (TPair<FGuid, TClientIdToPerRendererDataMap>& ClientPair : PerRendererDataMapsByClientId)
+	{
+		for (TClientIdToPerRendererDataMap::ElementType& PerRendererDataPair : ClientPair.Value)
+		{
+			FPerRendererData& PerRendererData = PerRendererDataPair.Value;
+			if (PerRendererData.TransactionId == InTransactionContext.TransactionId)
+			{
+				// Invalidate the transaction ID first so we don't try to end it a second time in EndLightCardDrag
+				PerRendererData.TransactionId.Invalidate();
+				EndLightCardDrag(PerRendererData, ClientPair.Key, PerRendererDataPair.Key, false);
+			}
+		}
+	}
+}
+#endif
+
 void FStageAppRouteHandler::HandleWebSocketNDisplayPreviewLightCardDragMove(const FRemoteControlWebSocketMessage& WebSocketMessage)
 {
 	FRCWebSocketNDisplayPreviewLightCardDragMoveBody Body;
@@ -664,7 +737,11 @@ void FStageAppRouteHandler::EndLightCardDrag(FPerRendererData& PerRendererData,
 	PerRendererData.bHasDragMovedRecently = false;
 
 #if WITH_EDITOR
-	GEditor->EndTransaction();
+	if (PerRendererData.TransactionId.IsValid())
+	{
+		PerRendererData.TransactionId.Invalidate(); // Invalidate this first so we don't try to re-handle the end event
+		GEditor->EndTransaction();
+	}
 #endif
 
 	if (!bEndedByClient)
diff --git a/Engine/Plugins/VirtualProduction/EpicStageApp/Source/EpicStageApp/Private/StageAppRouteHandler.h b/Engine/Plugins/VirtualProduction/EpicStageApp/Source/EpicStageApp/Private/StageAppRouteHandler.h
--- a/Engine/Plugins/VirtualProduction/EpicStageApp/Source/EpicStageApp/Private/StageAppRouteHandler.h
+++ b/Engine/Plugins/VirtualProduction/EpicStageApp/Source/EpicStageApp/Private/StageAppRouteHandler.h
@@ -7,6 +7,7 @@
 #include "DisplayClusterLightCardEditorHelper.h"
 #include "UObject/ObjectMacros.h"
 #include "IWebRemoteControlModule.h"
+#include "Misc/ITransaction.h"
 #include "RemoteControlWebsocketRoute.h"
 #include "StageAppRequest.h"
 
@@ -34,6 +35,11 @@ private:
 		/** If true, the dragged lightcards have moved since the last timeout check. */
 		bool bHasDragMovedRecently;
 
+#if WITH_EDITOR
+		/** The GUID of the current drag transaction. */
+		FGuid TransactionId;
+#endif		
+
 	public:
 		/** Get or create a FDisplayClusterLightCardEditorHelper for this client preview renderer. */
 		FDisplayClusterLightCardEditorHelper& GetLightCardHelper();
@@ -84,6 +90,12 @@ public:
 	/** Unregisters the WebSocket routes with the provided WebRemoteControl module. */
 	void UnregisterRoutes(IWebRemoteControlModule& WebRemoteControl);
 
+	/** Registers events that require the engine to be initialized first. */
+	void RegisterEngineEvents();
+
+	/** Unregisters events that required the engine to be initialized first. */
+	void UnregisterEngineEvents();
+
 private:
 	/** Register a WebSocket route with the WebRemoteControl module. */
 	void RegisterRoute(TUniquePtr<FRemoteControlWebsocketRoute> Route);
@@ -115,6 +127,13 @@ private:
 	/** Called when a client disconnects from the WebSocket server. */
 	void HandleClientDisconnected(FGuid ClientId);
 
+#if WITH_EDITOR
+	/**
+	 * Called when the state of an editor transaction changes.
+	 */
+	void OnTransactionStateChanged(const FTransactionContext& InTransactionContext, const ETransactionStateEventType InTransactionState);
+#endif
+
 	/** Change the settings for a client's preview renderer. */
 	void ChangePreviewRendererSettings(const FGuid& ClientId, int32 RendererId, const FRCWebSocketNDisplayPreviewRendererSettings& NewSettings);
 
diff --git a/Engine/Plugins/VirtualProduction/RemoteControl/Source/WebRemoteControl/Private/WebSocketMessageHandler.cpp b/Engine/Plugins/VirtualProduction/RemoteControl/Source/WebRemoteControl/Private/WebSocketMessageHandler.cpp
--- a/Engine/Plugins/VirtualProduction/RemoteControl/Source/WebRemoteControl/Private/WebSocketMessageHandler.cpp
+++ b/Engine/Plugins/VirtualProduction/RemoteControl/Source/WebRemoteControl/Private/WebSocketMessageHandler.cpp
@@ -273,24 +273,16 @@ void FWebSocketMessageHandler::RegisterRoutes(FWebRemoteControlModule* WebRemote
 	FCoreUObjectDelegates::OnObjectTransacted.AddRaw(this, &FWebSocketMessageHandler::OnObjectTransacted);
 
 	Server->OnConnectionClosed().AddRaw(this, &FWebSocketMessageHandler::OnConnectionClosedCallback);
+#endif
 
 	if (GEngine)
 	{
-		RegisterActorHandlers();
+		RegisterEngineEvents();
 	}
 	else
 	{
-		FCoreDelegates::OnPostEngineInit.AddRaw(this, &FWebSocketMessageHandler::RegisterActorHandlers);
-	}
-
-	if (GEditor)
-	{
-		if (UTransBuffer* TransBuffer = Cast<UTransBuffer>(GEditor->Trans))
-		{
-			TransBuffer->OnTransactionStateChanged().AddRaw(this, &FWebSocketMessageHandler::HandleTransactionStateChanged);
-		}
+		FCoreDelegates::OnPostEngineInit.AddRaw(this, &FWebSocketMessageHandler::RegisterEngineEvents);
 	}
-#endif
 	
 	// WebSocket routes
 	RegisterRoute(WebRemoteControl, MakeUnique<FRemoteControlWebsocketRoute>(
@@ -388,9 +380,17 @@ void FWebSocketMessageHandler::RegisterRoute(FWebRemoteControlModule* WebRemoteC
 	Routes.Emplace(MoveTemp(Route));
 }
 
-void FWebSocketMessageHandler::RegisterActorHandlers()
+void FWebSocketMessageHandler::RegisterEngineEvents()
 {
 #if WITH_EDITOR
+	if (GEditor)
+	{
+		if (UTransBuffer* TransBuffer = Cast<UTransBuffer>(GEditor->Trans))
+		{
+			TransBuffer->OnTransactionStateChanged().AddRaw(this, &FWebSocketMessageHandler::HandleTransactionStateChanged);
+		}
+	}
+
 	if (GEngine)
 	{
 		OnActorAddedHandle = GEngine->OnLevelActorAdded().AddRaw(this, &FWebSocketMessageHandler::OnActorAdded);
diff --git a/Engine/Plugins/VirtualProduction/RemoteControl/Source/WebRemoteControl/Private/WebSocketMessageHandler.h b/Engine/Plugins/VirtualProduction/RemoteControl/Source/WebRemoteControl/Private/WebSocketMessageHandler.h
--- a/Engine/Plugins/VirtualProduction/RemoteControl/Source/WebRemoteControl/Private/WebSocketMessageHandler.h
+++ b/Engine/Plugins/VirtualProduction/RemoteControl/Source/WebRemoteControl/Private/WebSocketMessageHandler.h
@@ -82,8 +82,8 @@ private:
 	/** Register a WebSocket route */
 	void RegisterRoute(FWebRemoteControlModule* WebRemoteControl, TUniquePtr<FRemoteControlWebsocketRoute> Route);
 
-	/** Register handlers for actors being added/deleted (must happen after engine init). */
-	void RegisterActorHandlers();
+	/** Registers events that require the engine to be initialized first. */
+	void RegisterEngineEvents();
 	
 	/** Handles registration to callbacks to a given preset */
 	void HandleWebSocketPresetRegister(const FRemoteControlWebSocketMessage& WebSocketMessage);
