EpicStage: Fixed missing struct member initializations.

diff --git a/Engine/Plugins/VirtualProduction/EpicStageApp/Source/EpicStageApp/Private/StageAppRequest.h b/Engine/Plugins/VirtualProduction/EpicStageApp/Source/EpicStageApp/Private/StageAppRequest.h
--- a/Engine/Plugins/VirtualProduction/EpicStageApp/Source/EpicStageApp/Private/StageAppRequest.h
+++ b/Engine/Plugins/VirtualProduction/EpicStageApp/Source/EpicStageApp/Private/StageAppRequest.h
@@ -279,7 +279,7 @@ struct FRCWebSocketNDisplayPreviewLightCardDragMoveBody : public FRCRequest
 	 * The position of the drag cursor in coordinates normalized to the size of the preview image.
 	 */
 	UPROPERTY()
-	FVector2D DragPosition;
+	FVector2D DragPosition = FVector2D::ZeroVector;
 
 	/**
 	 * The sequence number of this change. The highest sequence number received from this client will be
@@ -317,7 +317,7 @@ struct FRCWebSocketNDisplayPreviewLightCardDragEndBody : public FRCRequest
 	 * The final position of the drag cursor in coordinates normalized to the size of the preview image.
 	 */
 	UPROPERTY()
-	FVector2D DragPosition;
+	FVector2D DragPosition = FVector2D::ZeroVector;
 
 	/**
 	 * The sequence number of this change. The highest sequence number received from this client will be
diff --git a/Engine/Plugins/VirtualProduction/EpicStageApp/Source/EpicStageApp/Private/StageAppResponse.h b/Engine/Plugins/VirtualProduction/EpicStageApp/Source/EpicStageApp/Private/StageAppResponse.h
--- a/Engine/Plugins/VirtualProduction/EpicStageApp/Source/EpicStageApp/Private/StageAppResponse.h
+++ b/Engine/Plugins/VirtualProduction/EpicStageApp/Source/EpicStageApp/Private/StageAppResponse.h
@@ -42,7 +42,7 @@ struct FRCPreviewRenderCompletedEventActorPosition
 
 	/** The actor's position on the preview texture in coordinates normalized to the size of the preview render's texture. */
 	UPROPERTY()
-	FVector2D Position;
+	FVector2D Position = FVector2D::ZeroVector;
 };
 
 /**
diff --git a/Engine/Source/Runtime/VirtualProduction/StageDataCore/Public/StageMessages.h b/Engine/Source/Runtime/VirtualProduction/StageDataCore/Public/StageMessages.h
--- a/Engine/Source/Runtime/VirtualProduction/StageDataCore/Public/StageMessages.h
+++ b/Engine/Source/Runtime/VirtualProduction/StageDataCore/Public/StageMessages.h
@@ -283,7 +283,7 @@ public:
 
 	/** Event for this critical state */
 	UPROPERTY(VisibleAnywhere, Category = "AssetLoading")
-	EStageLoadingState LoadingState;
+	EStageLoadingState LoadingState = EStageLoadingState::PreLoad;
 
 	/** Name of the asset currently loading. */
 	UPROPERTY(VisibleAnywhere, Category = "AssetLoading")
