[EpicStageApp] Change normal debugging from a define to a CVar

diff --git a/Engine/Plugins/VirtualProduction/EpicStageApp/Source/EpicStageApp/Private/StageAppRouteHandler.cpp b/Engine/Plugins/VirtualProduction/EpicStageApp/Source/EpicStageApp/Private/StageAppRouteHandler.cpp
--- a/Engine/Plugins/VirtualProduction/EpicStageApp/Source/EpicStageApp/Private/StageAppRouteHandler.cpp
+++ b/Engine/Plugins/VirtualProduction/EpicStageApp/Source/EpicStageApp/Private/StageAppRouteHandler.cpp
@@ -21,10 +21,14 @@
 #include "Editor/TransBuffer.h"
 #endif
 
-#define DEBUG_STAGEAPP_NORMALS false
-
 #define LOCTEXT_NAMESPACE "StageAppRouteHandler"
 
+static TAutoConsoleVariable<bool> CVarStageAppDebugLightCardHelperNormals(
+	TEXT("StageApp.DebugLightCardHelperNormals"),
+	false,
+	TEXT("Whether to overlay preview renders with a debug display of the normal maps for the lightcard helper.")
+);
+
 static TAutoConsoleVariable<float> CVarStageAppDragTimeoutCheckInterval(
 	TEXT("StageApp.DragTimeoutCheckInterval"),
 	2.f,
@@ -406,28 +410,29 @@ void FStageAppRouteHandler::HandleWebSocketNDisplayPreviewRender(const FRemoteCo
 		{
 			FPerRendererData* PerRendererData = GetClientPerRendererData(ClientId, RendererId);
 
-#if DEBUG_STAGEAPP_NORMALS
-			// Draw the normal texture in the top-left corner if available
-			if (PerRendererData)
+			if (CVarStageAppDebugLightCardHelperNormals.GetValueOnGameThread())
 			{
-				FDisplayClusterLightCardEditorHelper& Helper = PerRendererData->GetLightCardHelper();
-				const FGameTime Time = FGameTime::GetTimeSinceAppStart();
-				const ERHIFeatureLevel::Type FeatureLevel = PerRendererData->GetRootActor()->GetWorld()->Scene->GetFeatureLevel();
-				const FIntPoint HalfResolution = Resolution / 2;
+				// Draw the normal texture in the top-left corner if available
+				if (PerRendererData)
+				{
+					FDisplayClusterLightCardEditorHelper& Helper = PerRendererData->GetLightCardHelper();
+					const FGameTime Time = FGameTime::GetTimeSinceAppStart();
+					const ERHIFeatureLevel::Type FeatureLevel = PerRendererData->GetRootActor()->GetWorld()->Scene->GetFeatureLevel();
+					const FIntPoint HalfResolution = Resolution / 2;
 
-				auto DrawNormalMap = [&](const FIntPoint& Position, bool bShowNorthMap) {
-					if (const UTexture2D* NorthMapTexture = Helper.GetNormalMapTexture(bShowNorthMap))
-					{
-						FCanvas Canvas(&RenderTarget, nullptr, Time, FeatureLevel);
-						Canvas.DrawTile(Position.X, Position.Y, HalfResolution.X, HalfResolution.Y, 0, 0, 1, 1, FColor::White, NorthMapTexture->GetResource());
-						Canvas.Flush_GameThread();
-					}
-				};
+					auto DrawNormalMap = [&](const FIntPoint& Position, bool bShowNorthMap) {
+						if (const UTexture2D* NorthMapTexture = Helper.GetNormalMapTexture(bShowNorthMap))
+						{
+							FCanvas Canvas(&RenderTarget, nullptr, Time, FeatureLevel);
+							Canvas.DrawTile(Position.X, Position.Y, HalfResolution.X, HalfResolution.Y, 0, 0, 1, 1, FColor::White, NorthMapTexture->GetResource());
+							Canvas.Flush_GameThread();
+						}
+					};
 
-				DrawNormalMap(FIntPoint(0, 0), true);
-				DrawNormalMap(FIntPoint(HalfResolution.X, 0), false);
+					DrawNormalMap(FIntPoint(0, 0), true);
+					DrawNormalMap(FIntPoint(HalfResolution.X, 0), false);
+				}
 			}
-#endif
 
 			// Read image data
 			TArray<FColor> PixelData;
