[EpicStageApp] Camera tumble controls

diff --git a/Engine/Plugins/VirtualProduction/EpicStageApp/Source/EpicStageApp/Private/StageAppRouteHandler.cpp b/Engine/Plugins/VirtualProduction/EpicStageApp/Source/EpicStageApp/Private/StageAppRouteHandler.cpp
--- a/Engine/Plugins/VirtualProduction/EpicStageApp/Source/EpicStageApp/Private/StageAppRouteHandler.cpp
+++ b/Engine/Plugins/VirtualProduction/EpicStageApp/Source/EpicStageApp/Private/StageAppRouteHandler.cpp
@@ -212,6 +212,12 @@ void FStageAppRouteHandler::RegisterRoutes(IWebRemoteControlModule& WebRemoteCon
 		FWebSocketMessageDelegate::CreateRaw(this, &FStageAppRouteHandler::HandleWebSocketNDisplayPreviewRendererDestroy)
 	));
 
+	RegisterRoute(MakeUnique<FRemoteControlWebsocketRoute>(
+		TEXT("Change the configuration of an nDisplay preview renderer. Old values will be retained for any properties not explicitly included."),
+		TEXT("ndisplay.preview.renderer.configure"),
+		FWebSocketMessageDelegate::CreateRaw(this, &FStageAppRouteHandler::HandleWebSocketNDisplayPreviewRendererConfigure)
+		));
+
 	RegisterRoute(MakeUnique<FRemoteControlWebsocketRoute>(
 		TEXT("Request a preview render from an nDisplay preview renderer"),
 		TEXT("ndisplay.preview.render"),
@@ -340,6 +346,18 @@ void FStageAppRouteHandler::HandleWebSocketNDisplayPreviewRendererConfigure(cons
 		return;
 	}
 
+	FPerRendererData* PerRendererData = GetClientPerRendererData(WebSocketMessage.ClientId, Body.RendererId);
+	if (PerRendererData)
+	{
+		// If we have existing settings for this renderer, copy its old settings and deserialize again over them. This lets us retain any
+		// values that weren't explicitly included in the WebSocket message.
+		Body.Settings = PerRendererData->GetPreviewSettings();
+		if (!WebRemoteControlUtils::DeserializeMessage(WebSocketMessage.RequestPayload, Body))
+		{
+			return;
+		}
+	}
+
 	ChangePreviewRendererSettings(WebSocketMessage.ClientId, Body.RendererId, Body.Settings);
 }
 
@@ -380,6 +398,15 @@ void FStageAppRouteHandler::HandleWebSocketNDisplayPreviewRender(const FRemoteCo
 		return;
 	}
 
+	const FRCWebSocketNDisplayPreviewRendererSettings& PreviewSettings = PerRendererData->GetPreviewSettings();
+	const FIntPoint Resolution = PreviewSettings.Resolution;
+
+	if (Resolution.GetMin() <= 0)
+	{
+		// Not a valid resolution, so skip this render
+		return;
+	}
+
 	// Set up the render settings
 	FDisplayClusterMeshProjectionRenderSettings RenderSettings;
 	if (!PerRendererData->GetSceneViewInitOptions(RenderSettings.ViewInitOptions))
@@ -387,7 +414,6 @@ void FStageAppRouteHandler::HandleWebSocketNDisplayPreviewRender(const FRemoteCo
 		return;
 	}
 
-	const FRCWebSocketNDisplayPreviewRendererSettings& PreviewSettings = PerRendererData->GetPreviewSettings();
 	RenderSettings.RenderType = StageAppRouteHandlerUtils::GetInternalPreviewRenderType(PreviewSettings.RenderType);
 	RenderSettings.EngineShowFlags.SetSelectionOutline(false);
 
@@ -398,7 +424,6 @@ void FStageAppRouteHandler::HandleWebSocketNDisplayPreviewRender(const FRemoteCo
 	EngineShowFlagOrthographicOverride(!bIsOrthographic, RenderSettings.EngineShowFlags);
 
 	FSceneViewInitOptions ViewInitOptions = RenderSettings.ViewInitOptions;
-	const FIntPoint Resolution = PreviewSettings.Resolution;
 	const FGuid ClientId = WebSocketMessage.ClientId;
 	const int32 RendererId = Body.RendererId;
 	const int32 JpegQuality = PreviewSettings.JpegQuality;
