[EpicStageApp] Rework UI to match new mockups

diff --git a/Engine/Plugins/VirtualProduction/EpicStageApp/Source/EpicStageApp/Private/StageAppResponse.h b/Engine/Plugins/VirtualProduction/EpicStageApp/Source/EpicStageApp/Private/StageAppResponse.h
--- a/Engine/Plugins/VirtualProduction/EpicStageApp/Source/EpicStageApp/Private/StageAppResponse.h
+++ b/Engine/Plugins/VirtualProduction/EpicStageApp/Source/EpicStageApp/Private/StageAppResponse.h
@@ -76,6 +76,12 @@ struct FRCPreviewRenderCompletedEvent
 	UPROPERTY()
 	FString ImageBase64 = "";
 
+	/**
+	 * The resolution of the image included with this message.
+	 */
+	UPROPERTY()
+	FIntPoint Resolution = FIntPoint(1024, 1024);
+
 	/**
 	 * A list of projected positions within the preview render for each rendered actor.
 	 */
diff --git a/Engine/Plugins/VirtualProduction/EpicStageApp/Source/EpicStageApp/Private/StageAppRouteHandler.cpp b/Engine/Plugins/VirtualProduction/EpicStageApp/Source/EpicStageApp/Private/StageAppRouteHandler.cpp
--- a/Engine/Plugins/VirtualProduction/EpicStageApp/Source/EpicStageApp/Private/StageAppRouteHandler.cpp
+++ b/Engine/Plugins/VirtualProduction/EpicStageApp/Source/EpicStageApp/Private/StageAppRouteHandler.cpp
@@ -472,6 +472,7 @@ void FStageAppRouteHandler::HandleWebSocketNDisplayPreviewRender(const FRemoteCo
 			FRCPreviewRenderCompletedEvent Event;
 			Event.RendererId = RendererId;
 			Event.ImageBase64 = FBase64::Encode(JpegData.GetData(), JpegData.Num());
+			Event.Resolution = Resolution;
 			Event.IsRealTimeDisabled = !IDisplayClusterScenePreview::Get().IsRealTimePreviewEnabled();
 
 			if (PerRendererData)
