[EpicStageApp] Improve user feedback in situations where the app may not be able to connect, and add a manual connect option in the app as a fallback

diff --git a/Engine/Plugins/VirtualProduction/EpicStageApp/Source/EpicStageApp/Private/StageAppBeaconReceiver.cpp b/Engine/Plugins/VirtualProduction/EpicStageApp/Source/EpicStageApp/Private/StageAppBeaconReceiver.cpp
--- a/Engine/Plugins/VirtualProduction/EpicStageApp/Source/EpicStageApp/Private/StageAppBeaconReceiver.cpp
+++ b/Engine/Plugins/VirtualProduction/EpicStageApp/Source/EpicStageApp/Private/StageAppBeaconReceiver.cpp
@@ -40,13 +40,12 @@ void FStageAppBeaconReceiver::Startup()
 		.AsNonBlocking()
 		.AsReusable()
 		.BoundToPort(Settings.DiscoveryPort)
-		.JoinedToGroup(DiscoveryAddress) // joins default interface
 		.WithMulticastLoopback()
 		.WithReceiveBufferSize(256);
 
 	if (!Socket)
 	{
-		UE_LOG(LogRemoteControl, Warning, TEXT("StageAppBeaconResponder failed to create multicast socket"));
+		UE_LOG(LogRemoteControl, Warning, TEXT("StageAppBeaconReceiver failed to create multicast socket"));
 		return;
 	}
 
@@ -58,20 +57,26 @@ void FStageAppBeaconReceiver::Startup()
 		TArray<TSharedPtr<FInternetAddr>> LocapIps;
 		ISocketSubsystem::Get(PLATFORM_SOCKETSUBSYSTEM)->GetLocalAdapterAddresses(LocapIps);
 
-		for (TSharedPtr<FInternetAddr>& LocapIp : LocapIps)
+		for (TSharedPtr<FInternetAddr>& LocalIp : LocapIps)
 		{
-			if (!LocapIp.IsValid())
+			if (!LocalIp.IsValid())
 			{
 				continue;
 			}
 			
-			const bool bJoinedGroup = Socket->JoinMulticastGroup(*DiscoveryInternetAddr, *LocapIp);
+			const bool bJoinedGroup = Socket->JoinMulticastGroup(*DiscoveryInternetAddr, *LocalIp);
 
-			if (!bJoinedGroup)
+			if (bJoinedGroup)
 			{
-				UE_LOG(LogRemoteControl, Warning, 
-					TEXT("StageAppBeaconResponder failed to join multicast group '%s' on detected local interface '%s'"), 
-					*DiscoveryAddress.ToString(), *LocapIp->ToString(false));
+				UE_LOG(LogRemoteControl, Log,
+					TEXT("StageAppBeaconReceiver joined multicast group '%s' on detected local interface '%s'"),
+					*DiscoveryAddress.ToString(), *LocalIp->ToString(false));
+			}
+			else
+			{
+				UE_LOG(LogRemoteControl, Warning,
+					TEXT("StageAppBeaconReceiver failed to join multicast group '%s' on detected local interface '%s'"),
+					*DiscoveryAddress.ToString(), *LocalIp->ToString(false));
 			}
 		}
 	}
