[EpicStageApp] Fix app losing the root nDisplay cluster actor after it's rebuilt

diff --git a/Engine/Plugins/Runtime/nDisplay/Source/DisplayClusterScenePreview/Private/DisplayClusterScenePreviewModule.cpp b/Engine/Plugins/Runtime/nDisplay/Source/DisplayClusterScenePreview/Private/DisplayClusterScenePreviewModule.cpp
--- a/Engine/Plugins/Runtime/nDisplay/Source/DisplayClusterScenePreview/Private/DisplayClusterScenePreviewModule.cpp
+++ b/Engine/Plugins/Runtime/nDisplay/Source/DisplayClusterScenePreview/Private/DisplayClusterScenePreviewModule.cpp
@@ -515,6 +515,7 @@ bool FDisplayClusterScenePreviewModule::OnTick(float DeltaTime)
 				// We were provided a canvas for this render job, so use it if possible
 				if (!Job.Canvas.IsValid())
 				{
+					Job.ResultDelegate.Execute(nullptr);
 					continue;
 				}
 
@@ -522,11 +523,12 @@ bool FDisplayClusterScenePreviewModule::OnTick(float DeltaTime)
 				FRenderTarget* RenderTarget = Canvas->GetRenderTarget();
 				if (!RenderTarget)
 				{
+					Job.ResultDelegate.Execute(nullptr);
 					continue;
 				}
 
 				InternalRenderImmediate(*Config, Job.Settings, *Canvas);
-				Job.ResultDelegate.Execute(*RenderTarget);
+				Job.ResultDelegate.Execute(RenderTarget);
 				break;
 			}
 
@@ -556,7 +558,7 @@ bool FDisplayClusterScenePreviewModule::OnTick(float DeltaTime)
 				FCanvas Canvas(RenderTargetResource, nullptr, FGameTime::GetTimeSinceAppStart(), World->Scene->GetFeatureLevel());
 
 				InternalRenderImmediate(*Config, Job.Settings, Canvas);
-				Job.ResultDelegate.Execute(*RenderTargetResource);
+				Job.ResultDelegate.Execute(RenderTargetResource);
 				break;
 			}
 			
@@ -564,6 +566,7 @@ bool FDisplayClusterScenePreviewModule::OnTick(float DeltaTime)
 		}
 
 		// Config no longer exists, so try the next render
+		Job.ResultDelegate.Execute(nullptr);
 	}
 
 	if (RenderQueue.IsEmpty())
diff --git a/Engine/Plugins/Runtime/nDisplay/Source/DisplayClusterScenePreview/Public/IDisplayClusterScenePreview.h b/Engine/Plugins/Runtime/nDisplay/Source/DisplayClusterScenePreview/Public/IDisplayClusterScenePreview.h
--- a/Engine/Plugins/Runtime/nDisplay/Source/DisplayClusterScenePreview/Public/IDisplayClusterScenePreview.h
+++ b/Engine/Plugins/Runtime/nDisplay/Source/DisplayClusterScenePreview/Public/IDisplayClusterScenePreview.h
@@ -11,7 +11,7 @@
 
 class ADisplayClusterRootActor;
 
-DECLARE_DELEGATE_OneParam(FRenderResultDelegate, FRenderTarget&);
+DECLARE_DELEGATE_OneParam(FRenderResultDelegate, FRenderTarget*);
 
 /**
  * Interface for module containing tools for generation Display Cluster scene previews.
@@ -168,7 +168,7 @@ public:
 	 * @param RendererId The ID of the renderer as returned from CreateRenderer.
 	 * @param RenderSettings Settings controlling how the scene will be rendered.
 	 * @param Size The size of the image to produce. Note that whenever this changes for a given renderer, the underlying RenderTarget will be resized, which has a performance cost.
-	 * @param ResultDelegate The delegate to call when the render is complete. It will be passed a FRenderTarget containing the rendered preview.
+	 * @param ResultDelegate The delegate to call when the render is complete. It will be passed a FRenderTarget containing the rendered preview, or null if the render failed.
 	 */
 	virtual bool RenderQueued(int32 RendererId, FDisplayClusterMeshProjectionRenderSettings& RenderSettings, const FIntPoint& Size, FRenderResultDelegate ResultDelegate) = 0;
 
@@ -178,7 +178,7 @@ public:
 	 * @param RendererId The ID of the renderer as returned from CreateRenderer.
 	 * @param RenderSettings Settings controlling how the scene will be rendered.
 	 * @param Canvas The canvas to draw to. If this is invalid when the render is ready to start, the render will be skipped.
-	 * @param ResultDelegate The delegate to call when the render is complete. It will be passed a FRenderTarget containing the rendered preview.
+	 * @param ResultDelegate The delegate to call when the render is complete. It will be passed a FRenderTarget containing the rendered preview, or null if the render failed.
 	 */
 	virtual bool RenderQueued(int32 RendererId, FDisplayClusterMeshProjectionRenderSettings& RenderSettings, TWeakPtr<FCanvas> Canvas, FRenderResultDelegate ResultDelegate) = 0;
 
diff --git a/Engine/Plugins/VirtualProduction/EpicStageApp/Source/EpicStageApp/Private/StageAppRouteHandler.cpp b/Engine/Plugins/VirtualProduction/EpicStageApp/Source/EpicStageApp/Private/StageAppRouteHandler.cpp
--- a/Engine/Plugins/VirtualProduction/EpicStageApp/Source/EpicStageApp/Private/StageAppRouteHandler.cpp
+++ b/Engine/Plugins/VirtualProduction/EpicStageApp/Source/EpicStageApp/Private/StageAppRouteHandler.cpp
@@ -393,7 +393,7 @@ void FStageAppRouteHandler::HandleWebSocketNDisplayPreviewRender(const FRemoteCo
 	}
 
 	FPerRendererData* PerRendererData = GetClientPerRendererData(WebSocketMessage.ClientId, Body.RendererId);
-	if (!PerRendererData)
+	if (!PerRendererData || PerRendererData->bIsRenderPending)
 	{
 		return;
 	}
@@ -429,11 +429,23 @@ void FStageAppRouteHandler::HandleWebSocketNDisplayPreviewRender(const FRemoteCo
 	const int32 JpegQuality = PreviewSettings.JpegQuality;
 	const bool bIncludeActorPositions = PreviewSettings.IncludeActorPositions;
 
+	PerRendererData->bIsRenderPending = true;
+
 	// Render the image
 	IDisplayClusterScenePreview::Get().RenderQueued(Body.RendererId, RenderSettings, Resolution, FRenderResultDelegate::CreateLambda(
-		[this, ViewInitOptions, ClientId, RendererId, Resolution, JpegQuality, bIncludeActorPositions](FRenderTarget& RenderTarget)
+		[this, ViewInitOptions, ClientId, RendererId, Resolution, JpegQuality, bIncludeActorPositions](FRenderTarget* RenderTarget)
 		{
 			FPerRendererData* PerRendererData = GetClientPerRendererData(ClientId, RendererId);
+			if (PerRendererData)
+			{
+				PerRendererData->bIsRenderPending = false;
+			}
+
+			if (!RenderTarget)
+			{
+				// Render failed
+				return;
+			}
 
 			if (CVarStageAppDebugLightCardHelperNormals.GetValueOnGameThread())
 			{
@@ -448,7 +460,7 @@ void FStageAppRouteHandler::HandleWebSocketNDisplayPreviewRender(const FRemoteCo
 					auto DrawNormalMap = [&](const FIntPoint& Position, bool bShowNorthMap) {
 						if (const UTexture2D* NorthMapTexture = Helper.GetNormalMapTexture(bShowNorthMap))
 						{
-							FCanvas Canvas(&RenderTarget, nullptr, Time, FeatureLevel);
+							FCanvas Canvas(RenderTarget, nullptr, Time, FeatureLevel);
 							Canvas.DrawTile(Position.X, Position.Y, HalfResolution.X, HalfResolution.Y, 0, 0, 1, 1, FColor::White, NorthMapTexture->GetResource());
 							Canvas.Flush_GameThread();
 						}
@@ -461,7 +473,7 @@ void FStageAppRouteHandler::HandleWebSocketNDisplayPreviewRender(const FRemoteCo
 
 			// Read image data
 			TArray<FColor> PixelData;
-			RenderTarget.ReadPixels(PixelData);
+			RenderTarget->ReadPixels(PixelData);
 			FImageView ImageView(PixelData.GetData(), Resolution.X, Resolution.Y, 1, ERawImageFormat::BGRA8, EGammaSpace::Linear);
 
 			// Convert to JPEG
diff --git a/Engine/Plugins/VirtualProduction/EpicStageApp/Source/EpicStageApp/Private/StageAppRouteHandler.h b/Engine/Plugins/VirtualProduction/EpicStageApp/Source/EpicStageApp/Private/StageAppRouteHandler.h
--- a/Engine/Plugins/VirtualProduction/EpicStageApp/Source/EpicStageApp/Private/StageAppRouteHandler.h
+++ b/Engine/Plugins/VirtualProduction/EpicStageApp/Source/EpicStageApp/Private/StageAppRouteHandler.h
@@ -33,7 +33,10 @@ private:
 		FVector2D LastDragPosition;
 
 		/** If true, the dragged lightcards have moved since the last timeout check. */
-		bool bHasDragMovedRecently;
+		bool bHasDragMovedRecently = false;
+
+		/** If true, the client is currently waiting for a render to be completed. */
+		bool bIsRenderPending = false;
 
 #if WITH_EDITOR
 		/** The GUID of the current drag transaction. */
