[EpicStageApp] Control ColorCorrectWindows (and other IDisplayClusterStageActors) from the app

diff --git a/Engine/Plugins/Experimental/ColorCorrectRegions/Source/ColorCorrectRegions/Private/ColorCorrectWindow.cpp b/Engine/Plugins/Experimental/ColorCorrectRegions/Source/ColorCorrectRegions/Private/ColorCorrectWindow.cpp
--- a/Engine/Plugins/Experimental/ColorCorrectRegions/Source/ColorCorrectRegions/Private/ColorCorrectWindow.cpp
+++ b/Engine/Plugins/Experimental/ColorCorrectRegions/Source/ColorCorrectRegions/Private/ColorCorrectWindow.cpp
@@ -270,11 +270,31 @@ FDisplayClusterPositionalParams AColorCorrectionWindow::GetPositionalParams() co
 	return PositionalParams;
 }
 
-
 ADEPRECATED_ColorCorrectWindow::ADEPRECATED_ColorCorrectWindow(const FObjectInitializer& ObjectInitializer)
 	: Super(ObjectInitializer)
 {
 
 }
 
+void AColorCorrectionWindow::GetPositionalProperties(FPositionalPropertyArray& OutPropertyPairs) const
+{
+	void* Container = (void*)(&PositionalParams);
+
+	const TSet<FName>& PropertyNames = GetPositionalPropertyNames();
+	OutPropertyPairs.Reserve(PropertyNames.Num());
+
+	for (const FName& PropertyName : PropertyNames)
+	{
+		if (FProperty* Property = FindFProperty<FProperty>(FDisplayClusterPositionalParams::StaticStruct(), PropertyName))
+		{
+			OutPropertyPairs.Emplace(Container, Property);
+		}
+	}
+
+	if (FStructProperty* ParamsProperty = FindFProperty<FStructProperty>(GetClass(), GET_MEMBER_NAME_CHECKED(AColorCorrectionWindow, PositionalParams)))
+	{
+		OutPropertyPairs.Emplace((void*)this, ParamsProperty);
+	}
+}
+
 #undef NOTIFY_PARAM_SETTER
diff --git a/Engine/Plugins/Experimental/ColorCorrectRegions/Source/ColorCorrectRegions/Public/ColorCorrectWindow.h b/Engine/Plugins/Experimental/ColorCorrectRegions/Source/ColorCorrectRegions/Public/ColorCorrectWindow.h
--- a/Engine/Plugins/Experimental/ColorCorrectRegions/Source/ColorCorrectRegions/Public/ColorCorrectWindow.h
+++ b/Engine/Plugins/Experimental/ColorCorrectRegions/Source/ColorCorrectRegions/Public/ColorCorrectWindow.h
@@ -108,6 +108,8 @@ public:
 	virtual void SetPositionalParams(const FDisplayClusterPositionalParams& InParams) override;
 	UFUNCTION(BlueprintCallable, Category = Orientation)
 	virtual FDisplayClusterPositionalParams GetPositionalParams() const override;
+
+	virtual void GetPositionalProperties(FPositionalPropertyArray& OutPropertyPairs) const override;
 	// ~End IDisplayClusterStageActor interface
 	
 protected:
diff --git a/Engine/Plugins/Runtime/nDisplay/Source/DisplayCluster/Private/Game/EngineClasses/Scene/DisplayClusterLightCardActor.cpp b/Engine/Plugins/Runtime/nDisplay/Source/DisplayCluster/Private/Game/EngineClasses/Scene/DisplayClusterLightCardActor.cpp
--- a/Engine/Plugins/Runtime/nDisplay/Source/DisplayCluster/Private/Game/EngineClasses/Scene/DisplayClusterLightCardActor.cpp
+++ b/Engine/Plugins/Runtime/nDisplay/Source/DisplayCluster/Private/Game/EngineClasses/Scene/DisplayClusterLightCardActor.cpp
@@ -631,6 +631,22 @@ FVector2D ADisplayClusterLightCardActor::GetUVCoordinates() const
 	return UVCoordinates;
 }
 
+void ADisplayClusterLightCardActor::GetPositionalProperties(FPositionalPropertyArray& OutPropertyPairs) const
+{
+	void* Container = (void*)(this);
+
+	const TSet<FName>& PropertyNames = GetPositionalPropertyNames();
+	OutPropertyPairs.Reserve(PropertyNames.Num());
+
+	for (const FName& PropertyName : PropertyNames)
+	{
+		if (FProperty* Property = FindFProperty<FProperty>(GetClass(), PropertyName))
+		{
+			OutPropertyPairs.Emplace(Container, Property);
+		}
+	}
+}
+
 void ADisplayClusterLightCardActor::CreateComponentsForExtenders()
 {
 	IModularFeatures& ModularFeatures = IModularFeatures::Get();
diff --git a/Engine/Plugins/Runtime/nDisplay/Source/DisplayCluster/Public/DisplayClusterLightCardActor.h b/Engine/Plugins/Runtime/nDisplay/Source/DisplayCluster/Public/DisplayClusterLightCardActor.h
--- a/Engine/Plugins/Runtime/nDisplay/Source/DisplayCluster/Public/DisplayClusterLightCardActor.h
+++ b/Engine/Plugins/Runtime/nDisplay/Source/DisplayCluster/Public/DisplayClusterLightCardActor.h
@@ -147,6 +147,8 @@ public:
 
 	virtual void SetUVCoordinates(const FVector2D& InUVCoordinates) override;
 	virtual FVector2D GetUVCoordinates() const override;
+
+	virtual void GetPositionalProperties(FPositionalPropertyArray& OutPropertyPairs) const override;
 	// ~End IDisplayClusterStageActor interface
 	
 public:
diff --git a/Engine/Plugins/Runtime/nDisplay/Source/DisplayClusterLightCardEditor/Private/Viewport/DisplayClusterLightCardEditorViewportClient.cpp b/Engine/Plugins/Runtime/nDisplay/Source/DisplayClusterLightCardEditor/Private/Viewport/DisplayClusterLightCardEditorViewportClient.cpp
--- a/Engine/Plugins/Runtime/nDisplay/Source/DisplayClusterLightCardEditor/Private/Viewport/DisplayClusterLightCardEditorViewportClient.cpp
+++ b/Engine/Plugins/Runtime/nDisplay/Source/DisplayClusterLightCardEditor/Private/Viewport/DisplayClusterLightCardEditorViewportClient.cpp
@@ -2104,18 +2104,37 @@ void FDisplayClusterLightCardEditorViewportClient::PropagateActorTransform(const
 	{
 		AActor* LevelInstance = FoundProxy->LevelInstance.AsActorChecked();
 		LevelInstance->Modify();
+
+		IDisplayClusterStageActor* LevelInstanceStageActor = Cast<IDisplayClusterStageActor>(LevelInstance);
+		if (!ensure(LevelInstanceStageActor))
+		{
+			return;
+		}
+
+		IDisplayClusterStageActor::FPositionalPropertyArray ProxyPropertyPairs;
+		ActorProxy->GetPositionalProperties(ProxyPropertyPairs);
+
+		IDisplayClusterStageActor::FPositionalPropertyArray LevelInstancePropertyPairs;
+		LevelInstanceStageActor->GetPositionalProperties(LevelInstancePropertyPairs);
+
+		if (!ensure(ProxyPropertyPairs.Num() == LevelInstancePropertyPairs.Num()))
+		{
+			return;
+		}
 		
 		// Set the level instance property value to our proxy property value.
-		auto TryChangeProperty = [&](FName InPropertyName, TArray<const FProperty*>& InOutChangedProperties) -> void
-		{
-			if (const FProperty* Property = FindFProperty<FProperty>(LevelInstance->GetClass(), InPropertyName))
+		auto TryChangeProperty = [&](
+			const IDisplayClusterStageActor::FPropertyPair& ProxyPropertyPair,
+			const IDisplayClusterStageActor::FPropertyPair& LevelInstancePropertyPair,
+			TArray<const FProperty*>& InOutChangedProperties
+		) -> void
+		{
+			// Only change if values are different.
+			const FProperty* Property = ProxyPropertyPair.Value;
+			if (!Property->Identical_InContainer(ProxyPropertyPair.Key, LevelInstancePropertyPair.Key))
 			{
-				// Only change if values are different.
-				if (!Property->Identical_InContainer(ActorProxy.AsActorChecked(), LevelInstance))
-				{
-					Property->CopyCompleteValue_InContainer(LevelInstance, ActorProxy.AsActorChecked());
-					InOutChangedProperties.Add(Property);
-				}
+				Property->CopyCompleteValue_InContainer(LevelInstancePropertyPair.Key, ProxyPropertyPair.Key);
+				InOutChangedProperties.Add(Property);
 			}
 		};
 		
@@ -2130,13 +2149,12 @@ void FDisplayClusterLightCardEditorViewportClient::PropagateActorTransform(const
 
 		LevelInstance->SetActorTransform(ActorProxy.AsActorChecked()->GetTransform() * RALevelTransformNoScale);
 
-		const TSet<FName>& PropertyNames = ActorProxy->GetPositionalPropertyNames();
-
 		TArray<const FProperty*> ChangedProperties;
-		ChangedProperties.Reserve(PropertyNames.Num());
-		for (const FName& PropertyName : PropertyNames)
+		ChangedProperties.Reserve(ProxyPropertyPairs.Num());
+
+		for (int32 PropertyIndex = 0; PropertyIndex < ProxyPropertyPairs.Num(); ++PropertyIndex)
 		{
-			TryChangeProperty(PropertyName, ChangedProperties);
+			TryChangeProperty(ProxyPropertyPairs[PropertyIndex], LevelInstancePropertyPairs[PropertyIndex], ChangedProperties);
 		}
 		
 		// Snapshot the changed properties so multi-user can update while dragging.
diff --git a/Engine/Plugins/Runtime/nDisplay/Source/DisplayClusterScenePreview/Private/DisplayClusterLightCardEditorHelper.cpp b/Engine/Plugins/Runtime/nDisplay/Source/DisplayClusterScenePreview/Private/DisplayClusterLightCardEditorHelper.cpp
--- a/Engine/Plugins/Runtime/nDisplay/Source/DisplayClusterScenePreview/Private/DisplayClusterLightCardEditorHelper.cpp
+++ b/Engine/Plugins/Runtime/nDisplay/Source/DisplayClusterScenePreview/Private/DisplayClusterLightCardEditorHelper.cpp
@@ -1703,31 +1703,29 @@ void FDisplayClusterLightCardEditorHelper::PostEditChangePropertiesForMovedActor
 
 	Actor.AsActorChecked()->Modify();
 
-	auto ModifyLightCardProperty = [&](const FName& PropertyName) -> FProperty*
+	auto ModifyLightCardProperty = [&](const IDisplayClusterStageActor::FPropertyPair& PropertyPair) -> void
 	{
-		if (FProperty* Property = FindFProperty<FProperty>(Actor.AsActorChecked()->GetClass(), PropertyName))
-		{
-			// Broadcast the event directly instead of PostEditChangeProperty on the object itself, which would attempt to create unnecessary snapshots for each call
-			FPropertyChangedEvent PropertyChangedEvent(Property, EPropertyChangeType::Interactive);
-			FCoreUObjectDelegates::OnObjectPropertyChanged.Broadcast(Actor.AsActorChecked(), PropertyChangedEvent);
-
-			return Property;
-		}
-		return nullptr;
+		// Broadcast the event directly instead of PostEditChangeProperty on the object itself, which would attempt to create unnecessary snapshots for each call
+		FPropertyChangedEvent PropertyChangedEvent(PropertyPair.Value, EPropertyChangeType::Interactive);
+		FCoreUObjectDelegates::OnObjectPropertyChanged.Broadcast(Actor.AsActorChecked(), PropertyChangedEvent);
 	};
 
-	const TSet<FName>& PropertyNames = Actor->GetPositionalPropertyNames();
+	IDisplayClusterStageActor::FPositionalPropertyArray PropertyPairs;
+	Actor->GetPositionalProperties(PropertyPairs);
 	
 	TArray<const FProperty*> ChangedProperties;
-	ChangedProperties.Reserve(PropertyNames.Num());
-	for (const FName& PropertyName : PropertyNames)
+	ChangedProperties.Reserve(PropertyPairs.Num());
+
+	for (const IDisplayClusterStageActor::FPropertyPair& PropertyPair : PropertyPairs)
 	{
-		if (const FProperty* Property = ModifyLightCardProperty(PropertyName))
-		{
-			ChangedProperties.Add(Property);
-		}
+		ModifyLightCardProperty(PropertyPair);
+
+		ChangedProperties.Add(PropertyPair.Value);
 	}
 
 	SnapshotTransactionBuffer(Actor.AsActorChecked(), MakeArrayView(ChangedProperties.GetData(), ChangedProperties.Num()));
+
+	// Force the actor to update its position in a way that the multi-user server will see
+	Actor.AsActorChecked()->PostEditMove(false);
 }
 #endif
diff --git a/Engine/Plugins/Runtime/nDisplay/Source/DisplayClusterScenePreview/Private/DisplayClusterScenePreviewModule.cpp b/Engine/Plugins/Runtime/nDisplay/Source/DisplayClusterScenePreview/Private/DisplayClusterScenePreviewModule.cpp
--- a/Engine/Plugins/Runtime/nDisplay/Source/DisplayClusterScenePreview/Private/DisplayClusterScenePreviewModule.cpp
+++ b/Engine/Plugins/Runtime/nDisplay/Source/DisplayClusterScenePreview/Private/DisplayClusterScenePreviewModule.cpp
@@ -458,11 +458,30 @@ void FDisplayClusterScenePreviewModule::AutoPopulateScene(FRendererConfig& Rende
 			TSet<ADisplayClusterLightCardActor*> LightCards;
 			UDisplayClusterBlueprintLib::FindLightCardsForRootActor(RootActor, LightCards);
 
+			TSet<AActor*> Actors;
+			Actors.Reserve(LightCards.Num());
 			for (ADisplayClusterLightCardActor* LightCard : LightCards)
 			{
-				RendererConfig.Renderer->AddActor(LightCard);
-				RendererConfig.AddedActors.Add(LightCard);
-				RendererConfig.AutoLightcards.Add(LightCard);
+				Actors.Add(LightCard);
+			}
+
+			// Also check for any non-lightcard actors in the world that are valid to control from ICVFX editors
+			if (UWorld* World = RootActor->GetWorld())
+			{
+				for (const TWeakObjectPtr<AActor> WeakActor : TActorRange<AActor>(World))
+				{
+					if (WeakActor.IsValid() && WeakActor->Implements<UDisplayClusterStageActor>())
+					{
+						Actors.Add(WeakActor.Get());
+					}
+				}
+			}
+
+			for (AActor* Actor : Actors)
+			{
+				RendererConfig.Renderer->AddActor(Actor);
+				RendererConfig.AddedActors.Add(Actor);
+				RendererConfig.AutoActors.Add(Actor);
 			}
 		}
 	}
@@ -595,7 +614,7 @@ void FDisplayClusterScenePreviewModule::OnLevelActorDeleted(AActor* Actor)
 	for (TPair<int32, FRendererConfig>& ConfigPair : RendererConfigs)
 	{
 		FRendererConfig& Config = ConfigPair.Value;
-		if (Config.AutoLightcards.Contains(Actor))
+		if (Config.AutoActors.Contains(Actor))
 		{
 			Config.bIsSceneDirty = true;
 		}
diff --git a/Engine/Plugins/Runtime/nDisplay/Source/DisplayClusterScenePreview/Private/DisplayClusterScenePreviewModule.h b/Engine/Plugins/Runtime/nDisplay/Source/DisplayClusterScenePreview/Private/DisplayClusterScenePreviewModule.h
--- a/Engine/Plugins/Runtime/nDisplay/Source/DisplayClusterScenePreview/Private/DisplayClusterScenePreviewModule.h
+++ b/Engine/Plugins/Runtime/nDisplay/Source/DisplayClusterScenePreview/Private/DisplayClusterScenePreviewModule.h
@@ -59,8 +59,8 @@ private:
 		/** All actors that have been added to the renderer (except for the root actor). */
 		TArray<TWeakObjectPtr<AActor>> AddedActors;
 
-		/** Lightcards that have been automatically added to the scene. */
-		TArray<TWeakObjectPtr<ADisplayClusterLightCardActor>> AutoLightcards;
+		/** Actors that have been automatically added to the scene. */
+		TArray<TWeakObjectPtr<AActor>> AutoActors;
 
 		/** If true, automatically update the renderer with lightcards belonging to the root actor. */
 		bool bAutoUpdateLightcards = false;
diff --git a/Engine/Plugins/Runtime/nDisplayModularFeatures/Source/DisplayClusterLightCardExtender/Public/StageActor/IDisplayClusterStageActor.h b/Engine/Plugins/Runtime/nDisplayModularFeatures/Source/DisplayClusterLightCardExtender/Public/StageActor/IDisplayClusterStageActor.h
--- a/Engine/Plugins/Runtime/nDisplayModularFeatures/Source/DisplayClusterLightCardExtender/Public/StageActor/IDisplayClusterStageActor.h
+++ b/Engine/Plugins/Runtime/nDisplayModularFeatures/Source/DisplayClusterLightCardExtender/Public/StageActor/IDisplayClusterStageActor.h
@@ -18,6 +18,13 @@ class DISPLAYCLUSTERLIGHTCARDEXTENDER_API IDisplayClusterStageActor
 {
 	GENERATED_BODY()
 
+public:
+	/** A pair made up of a container pointer and a property within that container */
+	typedef TPair<void*, FProperty*> FPropertyPair;
+
+	/** Array type used to return positional properties for a stage actor */
+	typedef TArray<FPropertyPair, TInlineAllocator<16>> FPositionalPropertyArray;
+
 public:
 	/** The rotation used to orient the plane mesh used for the light card so that its normal points radially inwards */
 	static const FRotator PlaneMeshRotation;
@@ -52,6 +59,9 @@ public:
 	/** Return property names defined for this stage actor for use in property notifies */
 	virtual const TSet<FName>& GetPositionalPropertyNames() const;
 
+	/** Get an array of positional properties defined for this stage actor for use in property notifies. */
+	virtual void GetPositionalProperties(FPositionalPropertyArray& OutPropertyPairs) const = 0;
+
 	/** Return the bounding box for the stage actor. The default implementation returns the components bounding box */
 	virtual FBox GetBoxBounds(bool bLocalSpace = false) const;
 
diff --git a/Engine/Plugins/VirtualProduction/EpicStageApp/Source/EpicStageApp/Private/StageAppRequest.h b/Engine/Plugins/VirtualProduction/EpicStageApp/Source/EpicStageApp/Private/StageAppRequest.h
--- a/Engine/Plugins/VirtualProduction/EpicStageApp/Source/EpicStageApp/Private/StageAppRequest.h
+++ b/Engine/Plugins/VirtualProduction/EpicStageApp/Source/EpicStageApp/Private/StageAppRequest.h
@@ -208,14 +208,14 @@ struct FRCWebSocketNDisplayPreviewRenderBody : public FRCRequest
 };
 
 /**
- * Holds a request made via websocket to start dragging one or more light cards.
+ * Holds a request made via websocket to start dragging one or more actors.
  */
 USTRUCT()
-struct FRCWebSocketNDisplayPreviewLightCardDragBeginBody : public FRCRequest
+struct FRCWebSocketNDisplayPreviewActorDragBeginBody : public FRCRequest
 {
 	GENERATED_BODY()
 
-	FRCWebSocketNDisplayPreviewLightCardDragBeginBody()
+	FRCWebSocketNDisplayPreviewActorDragBeginBody()
 	{
 		AddStructParameter(ParametersFieldLabel());
 	}
@@ -232,16 +232,16 @@ struct FRCWebSocketNDisplayPreviewLightCardDragBeginBody : public FRCRequest
 	int32 RendererId = -1;
 
 	/**
-	 * Paths of the light cards that will be included in this drag action.
+	 * Paths of the actors that will be included in this drag action.
 	 */
 	UPROPERTY()
-	TArray<FString> LightCards;
+	TArray<FString> Actors;
 
 	/**
-	 * The light card to use as the origin point for the drag.
+	 * The actor to use as the origin point for the drag.
 	 */
 	UPROPERTY()
-	FString PrimaryLightCard;
+	FString PrimaryActor;
 
 	/**
 	 * The sequence number of this change. The highest sequence number received from this client will be
@@ -253,14 +253,14 @@ struct FRCWebSocketNDisplayPreviewLightCardDragBeginBody : public FRCRequest
 
 
 /**
- * Holds a request made via websocket to move the light cards that are currently being dragged for a preview renderer.
+ * Holds a request made via websocket to move the actors that are currently being dragged for a preview renderer.
  */
 USTRUCT()
-struct FRCWebSocketNDisplayPreviewLightCardDragMoveBody : public FRCRequest
+struct FRCWebSocketNDisplayPreviewActorDragMoveBody : public FRCRequest
 {
 	GENERATED_BODY()
 
-	FRCWebSocketNDisplayPreviewLightCardDragMoveBody()
+	FRCWebSocketNDisplayPreviewActorDragMoveBody()
 	{
 		AddStructParameter(ParametersFieldLabel());
 	}
@@ -291,14 +291,14 @@ struct FRCWebSocketNDisplayPreviewLightCardDragMoveBody : public FRCRequest
 };
 
 /**
- * Holds a request made via websocket to stop dragging light cards for a preview renderer.
+ * Holds a request made via websocket to stop dragging actors for a preview renderer.
  */
 USTRUCT()
-struct FRCWebSocketNDisplayPreviewLightCardDragEndBody : public FRCRequest
+struct FRCWebSocketNDisplayPreviewActorDragEndBody : public FRCRequest
 {
 	GENERATED_BODY()
 
-	FRCWebSocketNDisplayPreviewLightCardDragEndBody()
+	FRCWebSocketNDisplayPreviewActorDragEndBody()
 	{
 		AddStructParameter(ParametersFieldLabel());
 	}
diff --git a/Engine/Plugins/VirtualProduction/EpicStageApp/Source/EpicStageApp/Private/StageAppResponse.h b/Engine/Plugins/VirtualProduction/EpicStageApp/Source/EpicStageApp/Private/StageAppResponse.h
--- a/Engine/Plugins/VirtualProduction/EpicStageApp/Source/EpicStageApp/Private/StageAppResponse.h
+++ b/Engine/Plugins/VirtualProduction/EpicStageApp/Source/EpicStageApp/Private/StageAppResponse.h
@@ -89,13 +89,13 @@ struct FRCPreviewRenderCompletedEvent
 	TArray<FRCPreviewRenderCompletedEventActorPosition> ActorPositions;
 
 	/**
-	 * The last sequence number that the client sent during a light card drag operation on this renderer.
+	 * The last sequence number that the client sent during a stage actor drag operation on this renderer.
 	 */
 	UPROPERTY()
 	int64 SequenceNumber = -1;
 
 	/**
-	 * If true, real-time preview rendering is disabled in the engine, which may prevent nDisplay previews and light card positions from updating correctly.
+	 * If true, real-time preview rendering is disabled in the engine, which may prevent nDisplay previews and stage actor positions from updating correctly.
 	 */
 	UPROPERTY()
 	bool IsRealTimeDisabled = false;
@@ -105,12 +105,12 @@ struct FRCPreviewRenderCompletedEvent
  * Event sent to a client when a drag operation is cancelled by the engine.
  */
 USTRUCT()
-struct FRCLightCardDragCancelled
+struct FRCActorDragCancelled
 {
 	GENERATED_BODY()
 
-	FRCLightCardDragCancelled()
-	: Type(TEXT("LightCardDragCancelled"))
+	FRCActorDragCancelled()
+	: Type(TEXT("ActorDragCancelled"))
 	{
 	}
 
diff --git a/Engine/Plugins/VirtualProduction/EpicStageApp/Source/EpicStageApp/Private/StageAppRouteHandler.cpp b/Engine/Plugins/VirtualProduction/EpicStageApp/Source/EpicStageApp/Private/StageAppRouteHandler.cpp
--- a/Engine/Plugins/VirtualProduction/EpicStageApp/Source/EpicStageApp/Private/StageAppRouteHandler.cpp
+++ b/Engine/Plugins/VirtualProduction/EpicStageApp/Source/EpicStageApp/Private/StageAppRouteHandler.cpp
@@ -236,21 +236,21 @@ void FStageAppRouteHandler::RegisterRoutes(IWebRemoteControlModule& WebRemoteCon
 	));
 
 	RegisterRoute(MakeUnique<FRemoteControlWebsocketRoute>(
-		TEXT("Start dragging lightcards relative to a projected preview"),
-		TEXT("ndisplay.preview.lightcard.drag.begin"),
-		FWebSocketMessageDelegate::CreateRaw(this, &FStageAppRouteHandler::HandleWebSocketNDisplayPreviewLightCardDragBegin)
+		TEXT("Start dragging actors relative to a projected preview"),
+		TEXT("ndisplay.preview.actor.drag.begin"),
+		FWebSocketMessageDelegate::CreateRaw(this, &FStageAppRouteHandler::HandleWebSocketNDisplayPreviewActorDragBegin)
 	));
 
 	RegisterRoute(MakeUnique<FRemoteControlWebsocketRoute>(
-		TEXT("Move the light cards that are currently being dragged"),
-		TEXT("ndisplay.preview.lightcard.drag.move"),
-		FWebSocketMessageDelegate::CreateRaw(this, &FStageAppRouteHandler::HandleWebSocketNDisplayPreviewLightCardDragMove)
+		TEXT("Move the actors that are currently being dragged"),
+		TEXT("ndisplay.preview.actor.drag.move"),
+		FWebSocketMessageDelegate::CreateRaw(this, &FStageAppRouteHandler::HandleWebSocketNDisplayPreviewActorDragMove)
 	));
 
 	RegisterRoute(MakeUnique<FRemoteControlWebsocketRoute>(
-		TEXT("Finish dragging the light cards that are currently being dragged"),
-		TEXT("ndisplay.preview.lightcard.drag.end"),
-		FWebSocketMessageDelegate::CreateRaw(this, &FStageAppRouteHandler::HandleWebSocketNDisplayPreviewLightCardDragEnd)
+		TEXT("Finish dragging the actors that are currently being dragged"),
+		TEXT("ndisplay.preview.actor.drag.end"),
+		FWebSocketMessageDelegate::CreateRaw(this, &FStageAppRouteHandler::HandleWebSocketNDisplayPreviewActorDragEnd)
 	));
 }
 
@@ -524,15 +524,15 @@ void FStageAppRouteHandler::HandleWebSocketNDisplayPreviewRender(const FRemoteCo
 							}
 
 							FVector WorldPosition;
-							if (const ADisplayClusterLightCardActor* LightCard = Cast<ADisplayClusterLightCardActor>(Actor))
+							if (const IDisplayClusterStageActor* StageActor = Cast<IDisplayClusterStageActor>(Actor))
 							{
-								// Only send UV lightcards in UV mode, and only send non-UV lightcards in non-UV modes
-								if (LightCard->bIsUVLightCard != bIsInUVMode)
+								// Only send UV actors in UV mode, and only send non-UV actors in non-UV modes
+								if (StageActor->IsUVActor() != bIsInUVMode)
 								{
 									continue;
 								}
 
-								WorldPosition = LightCard->GetStageActorTransform().GetTranslation();
+								WorldPosition = StageActor->GetStageActorTransform().GetTranslation();
 							}
 							else
 							{
@@ -566,9 +566,9 @@ void FStageAppRouteHandler::HandleWebSocketNDisplayPreviewRender(const FRemoteCo
 	));
 }
 
-void FStageAppRouteHandler::HandleWebSocketNDisplayPreviewLightCardDragBegin(const FRemoteControlWebSocketMessage& WebSocketMessage)
+void FStageAppRouteHandler::HandleWebSocketNDisplayPreviewActorDragBegin(const FRemoteControlWebSocketMessage& WebSocketMessage)
 {
-	FRCWebSocketNDisplayPreviewLightCardDragBeginBody Body;
+	FRCWebSocketNDisplayPreviewActorDragBeginBody Body;
 	if (!WebRemoteControlUtils::DeserializeMessage(WebSocketMessage.RequestPayload, Body))
 	{
 		return;
@@ -580,7 +580,7 @@ void FStageAppRouteHandler::HandleWebSocketNDisplayPreviewLightCardDragBegin(con
 		return;
 	}
 
-	if (!PerRendererData->DraggedLightCards.IsEmpty())
+	if (!PerRendererData->DraggedActors.IsEmpty())
 	{
 		// A drag is already in progress
 		return;
@@ -590,7 +590,7 @@ void FStageAppRouteHandler::HandleWebSocketNDisplayPreviewLightCardDragBegin(con
 	PerRendererData->TransactionId.Invalidate();
 	if (GEditor && GEditor->Trans)
 	{
-		if (GEditor->BeginTransaction(LOCTEXT("DragLightcardsTransaction", "Drag Light Cards with Stage App")) != INDEX_NONE)
+		if (GEditor->BeginTransaction(LOCTEXT("DragActorsTransaction", "Drag Actors with Stage App")) != INDEX_NONE)
 		{
 			const FTransaction* Transaction = GEditor->Trans->GetTransaction(GEditor->Trans->GetQueueLength() - 1);
 			if (ensure(Transaction))
@@ -604,17 +604,29 @@ void FStageAppRouteHandler::HandleWebSocketNDisplayPreviewLightCardDragBegin(con
 	PerRendererData->UpdateDragSequenceNumber(Body.SequenceNumber);
 	PerRendererData->bHasDragMovedRecently = true;
 
-	for (const FString& LightCardPath : Body.LightCards)
+	for (const FString& ActorPath : Body.Actors)
 	{
-		if (ADisplayClusterLightCardActor* Actor = FindObject<ADisplayClusterLightCardActor>(nullptr, *LightCardPath))
+		if (AActor* Actor = FindObject<AActor>(nullptr, *ActorPath))
 		{
-			PerRendererData->DraggedLightCards.Emplace(Actor);
+			if (!Actor || !Actor->Implements<UDisplayClusterStageActor>())
+			{
+				continue;
+			}
+
+			PerRendererData->DraggedActors.Emplace(FDisplayClusterWeakStageActorPtr(Actor));
 		}
 	}
-
-	if (!Body.PrimaryLightCard.IsEmpty())
+	
+	if (!Body.PrimaryActor.IsEmpty())
 	{
-		PerRendererData->PrimaryLightCard = FindObject<ADisplayClusterLightCardActor>(nullptr, *Body.PrimaryLightCard);
+		AActor* PrimaryActor = FindObject<AActor>(nullptr, *Body.PrimaryActor);
+
+		if (!PrimaryActor || !PrimaryActor->Implements<UDisplayClusterStageActor>())
+		{
+			return;
+		}
+
+		PerRendererData->PrimaryActor = FDisplayClusterWeakStageActorPtr(PrimaryActor);
 	}
 
 	if (!DragTimeoutTickerHandle.IsValid())
@@ -641,18 +653,18 @@ void FStageAppRouteHandler::OnTransactionStateChanged(const FTransactionContext&
 			FPerRendererData& PerRendererData = PerRendererDataPair.Value;
 			if (PerRendererData.TransactionId == InTransactionContext.TransactionId)
 			{
-				// Invalidate the transaction ID first so we don't try to end it a second time in EndLightCardDrag
+				// Invalidate the transaction ID first so we don't try to end it a second time in EndActorDrag
 				PerRendererData.TransactionId.Invalidate();
-				EndLightCardDrag(PerRendererData, ClientPair.Key, PerRendererDataPair.Key, false);
+				EndActorDrag(PerRendererData, ClientPair.Key, PerRendererDataPair.Key, false);
 			}
 		}
 	}
 }
 #endif
 
-void FStageAppRouteHandler::HandleWebSocketNDisplayPreviewLightCardDragMove(const FRemoteControlWebSocketMessage& WebSocketMessage)
+void FStageAppRouteHandler::HandleWebSocketNDisplayPreviewActorDragMove(const FRemoteControlWebSocketMessage& WebSocketMessage)
 {
-	FRCWebSocketNDisplayPreviewLightCardDragMoveBody Body;
+	FRCWebSocketNDisplayPreviewActorDragMoveBody Body;
 	if (!WebRemoteControlUtils::DeserializeMessage(WebSocketMessage.RequestPayload, Body))
 	{
 		return;
@@ -667,12 +679,12 @@ void FStageAppRouteHandler::HandleWebSocketNDisplayPreviewLightCardDragMove(cons
 	PerRendererData->UpdateDragSequenceNumber(Body.SequenceNumber);
 	PerRendererData->bHasDragMovedRecently = true;
 
-	DragLightCards(*PerRendererData, Body.DragPosition);
+	DragActors(*PerRendererData, Body.DragPosition);
 }
 
-void FStageAppRouteHandler::HandleWebSocketNDisplayPreviewLightCardDragEnd(const FRemoteControlWebSocketMessage& WebSocketMessage)
+void FStageAppRouteHandler::HandleWebSocketNDisplayPreviewActorDragEnd(const FRemoteControlWebSocketMessage& WebSocketMessage)
 {
-	FRCWebSocketNDisplayPreviewLightCardDragEndBody Body;
+	FRCWebSocketNDisplayPreviewActorDragEndBody Body;
 	if (!WebRemoteControlUtils::DeserializeMessage(WebSocketMessage.RequestPayload, Body))
 	{
 		return;
@@ -685,17 +697,17 @@ void FStageAppRouteHandler::HandleWebSocketNDisplayPreviewLightCardDragEnd(const
 	}
 
 	PerRendererData->UpdateDragSequenceNumber(Body.SequenceNumber);
-	DragLightCards(*PerRendererData, Body.DragPosition);
-	EndLightCardDrag(*PerRendererData, WebSocketMessage.ClientId, Body.RendererId, true);
+	DragActors(*PerRendererData, Body.DragPosition);
+	EndActorDrag(*PerRendererData, WebSocketMessage.ClientId, Body.RendererId, true);
 }
 
 void FStageAppRouteHandler::HandleClientDisconnected(FGuid ClientId)
 {
 	// Destroy the client's renderers
-	if (const TClientIdToPerRendererDataMap* PerRendererDataMap = PerRendererDataMapsByClientId.Find(ClientId))
+	if (TClientIdToPerRendererDataMap* PerRendererDataMap = PerRendererDataMapsByClientId.Find(ClientId))
 	{
 		IDisplayClusterScenePreview& PreviewModule = IDisplayClusterScenePreview::Get();
-		for (const TClientIdToPerRendererDataMap::ElementType& PerRendererDataPair : *PerRendererDataMap)
+		for (TClientIdToPerRendererDataMap::ElementType& PerRendererDataPair : *PerRendererDataMap)
 		{
 			PreviewModule.DestroyRenderer(PerRendererDataPair.Key);
 		}
@@ -726,7 +738,7 @@ FStageAppRouteHandler::FPerRendererData* FStageAppRouteHandler::GetClientPerRend
 	return PerRendererDataMap->Find(RendererId);
 }
 
-void FStageAppRouteHandler::DragLightCards(FPerRendererData& PerRendererData, FVector2D DragPosition)
+void FStageAppRouteHandler::DragActors(FPerRendererData& PerRendererData, FVector2D DragPosition)
 {
 	// Set up the scene view relative to which we're dragging
 	FSceneViewInitOptions ViewInitOptions;
@@ -749,43 +761,29 @@ void FStageAppRouteHandler::DragLightCards(FPerRendererData& PerRendererData, FV
 	);
 	const FVector DragWidgetOffset = FVector::ZeroVector;
 
-	// Light card editor supports any type of actor implementing IDisplayClusterStageActor which is managed by
-	// a custom weak ptr. Convert the LightCard actors to the actor weak ptr array.
-	TArray<FDisplayClusterWeakStageActorPtr> StageActors;
-	StageActors.Reserve(PerRendererData.DraggedLightCards.Num());
-	Algo::TransformIf(PerRendererData.DraggedLightCards, StageActors,
-		[](const TWeakObjectPtr<ADisplayClusterLightCardActor>& WeakPtr)
-		{
-			return WeakPtr.IsValid();
-		},
-		[](const TWeakObjectPtr<ADisplayClusterLightCardActor>& WeakPtr)
-		{
-			return FDisplayClusterWeakStageActorPtr(WeakPtr.Get());
-		});
-
 	FDisplayClusterLightCardEditorHelper& LightCardHelper = PerRendererData.GetLightCardHelper();
 
 	if (LightCardHelper.GetProjectionMode() == EDisplayClusterMeshProjectionType::UV)
 	{
 		LightCardHelper.DragUVActors(
-			MoveTemp(StageActors),
+			PerRendererData.DraggedActors,
 			PixelPos,
 			View,
 			DragWidgetOffset,
 			EAxisList::Type::XYZ,
-			PerRendererData.PrimaryLightCard.Get()
+			PerRendererData.PrimaryActor
 		);
 	}
 	else
 	{
 		LightCardHelper.DragActors(
-			MoveTemp(StageActors),
+			PerRendererData.DraggedActors,
 			PixelPos,
 			View,
 			FDisplayClusterLightCardEditorHelper::ECoordinateSystem::Spherical,
 			DragWidgetOffset,
 			EAxisList::Type::XYZ,
-			PerRendererData.PrimaryLightCard.Get()
+			PerRendererData.PrimaryActor
 		);
 	}
 }
@@ -798,14 +796,14 @@ bool FStageAppRouteHandler::TimeOutDrags(float DeltaTime)
 		for (TClientIdToPerRendererDataMap::ElementType& PerRendererDataPair : ClientPair.Value)
 		{
 			FPerRendererData& PerRendererData = PerRendererDataPair.Value;
-			if (PerRendererData.DraggedLightCards.IsEmpty())
+			if (PerRendererData.DraggedActors.IsEmpty())
 			{
 				continue;
 			}
 			
 			if (!PerRendererData.bHasDragMovedRecently)
 			{
-				EndLightCardDrag(PerRendererData, ClientPair.Key, PerRendererDataPair.Key, false);
+				EndActorDrag(PerRendererData, ClientPair.Key, PerRendererDataPair.Key, false);
 				continue;
 			}
 
@@ -823,9 +821,9 @@ bool FStageAppRouteHandler::TimeOutDrags(float DeltaTime)
 	return true;
 }
 
-void FStageAppRouteHandler::EndLightCardDrag(FPerRendererData& PerRendererData, const FGuid& ClientId, int32 RendererId, bool bEndedByClient)
+void FStageAppRouteHandler::EndActorDrag(FPerRendererData& PerRendererData, const FGuid& ClientId, int32 RendererId, bool bEndedByClient)
 {
-	PerRendererData.DraggedLightCards.Empty();
+	PerRendererData.DraggedActors.Empty();
 	PerRendererData.bHasDragMovedRecently = false;
 
 #if WITH_EDITOR
@@ -839,7 +837,7 @@ void FStageAppRouteHandler::EndLightCardDrag(FPerRendererData& PerRendererData,
 	if (!bEndedByClient)
 	{
 		// Notify the client that its drag in progress was cancelled
-		FRCLightCardDragCancelled Event;
+		FRCActorDragCancelled Event;
 		Event.RendererId = RendererId;
 
 		TArray<uint8> Payload;
diff --git a/Engine/Plugins/VirtualProduction/EpicStageApp/Source/EpicStageApp/Private/StageAppRouteHandler.h b/Engine/Plugins/VirtualProduction/EpicStageApp/Source/EpicStageApp/Private/StageAppRouteHandler.h
--- a/Engine/Plugins/VirtualProduction/EpicStageApp/Source/EpicStageApp/Private/StageAppRouteHandler.h
+++ b/Engine/Plugins/VirtualProduction/EpicStageApp/Source/EpicStageApp/Private/StageAppRouteHandler.h
@@ -23,11 +23,11 @@ private:
 	{
 		FPerRendererData(int RendererId);
 
-		/** The list of light cards currently being dragged by this client. Empty if a drag is not in progress. */
-		TArray<TWeakObjectPtr<ADisplayClusterLightCardActor>> DraggedLightCards;
+		/** The list of actors currently being dragged by this client. Empty if a drag is not in progress. */
+		TArray<FDisplayClusterWeakStageActorPtr> DraggedActors;
 
-		/** The light card to use as the origin point for the drag operation. */
-		TWeakObjectPtr<ADisplayClusterLightCardActor> PrimaryLightCard;
+		/** The actor to use as the origin point for the drag operation. */
+		FDisplayClusterWeakStageActorPtr PrimaryActor;
 
 		/** The last position sent by a client during a drag operation. */
 		FVector2D LastDragPosition;
@@ -118,14 +118,14 @@ private:
 	/** Handles rendering a preview. */
 	void HandleWebSocketNDisplayPreviewRender(const FRemoteControlWebSocketMessage& WebSocketMessage);
 
-	/** Handles beginning a light card drag operation. */
-	void HandleWebSocketNDisplayPreviewLightCardDragBegin(const FRemoteControlWebSocketMessage& WebSocketMessage);
+	/** Handles beginning an actor drag operation. */
+	void HandleWebSocketNDisplayPreviewActorDragBegin(const FRemoteControlWebSocketMessage& WebSocketMessage);
 
-	/** Handles moving light cards with a drag operation. */
-	void HandleWebSocketNDisplayPreviewLightCardDragMove(const FRemoteControlWebSocketMessage& WebSocketMessage);
+	/** Handles moving actors with a drag operation. */
+	void HandleWebSocketNDisplayPreviewActorDragMove(const FRemoteControlWebSocketMessage& WebSocketMessage);
 
-	/** Handles ending a light card drag operation. */
-	void HandleWebSocketNDisplayPreviewLightCardDragEnd(const FRemoteControlWebSocketMessage& WebSocketMessage);
+	/** Handles ending an actor drag operation. */
+	void HandleWebSocketNDisplayPreviewActorDragEnd(const FRemoteControlWebSocketMessage& WebSocketMessage);
 
 	/** Called when a client disconnects from the WebSocket server. */
 	void HandleClientDisconnected(FGuid ClientId);
@@ -143,14 +143,14 @@ private:
 	/** Get the nDisplay preview rendering data for the given client and renderer ID, or null if it's an invalid client/renderer ID. */
 	FPerRendererData* GetClientPerRendererData(const FGuid& ClientId, int32 RendererId);
 
-	/** Drag light cards to the specified position. */
-	void DragLightCards(FPerRendererData& PreviewData, FVector2D DragPosition);
+	/** Drag actors to the specified position. */
+	void DragActors(FPerRendererData& PreviewData, FVector2D DragPosition);
 
-	/** Check if any light card drag operations have timed out, and if so, end them. */
+	/** Check if any actor drag operations have timed out, and if so, end them. */
 	bool TimeOutDrags(float DeltaTime);
 
-	/** End a client's light card drag operation. */
-	void EndLightCardDrag(FPerRendererData& PreviewData, const FGuid& ClientId, int32 RendererId, bool bEndedByClient);
+	/** End a client's actor drag operation. */
+	void EndActorDrag(FPerRendererData& PreviewData, const FGuid& ClientId, int32 RendererId, bool bEndedByClient);
 
 private:
 	/** The module for the WebSocket server for which this is handling routes. */
