[EpicStageApp] Cancel drag when client disconnects

diff --git a/Engine/Plugins/VirtualProduction/EpicStageApp/Source/EpicStageApp/Private/StageAppRouteHandler.cpp b/Engine/Plugins/VirtualProduction/EpicStageApp/Source/EpicStageApp/Private/StageAppRouteHandler.cpp
--- a/Engine/Plugins/VirtualProduction/EpicStageApp/Source/EpicStageApp/Private/StageAppRouteHandler.cpp
+++ b/Engine/Plugins/VirtualProduction/EpicStageApp/Source/EpicStageApp/Private/StageAppRouteHandler.cpp
@@ -709,6 +709,7 @@ void FStageAppRouteHandler::HandleClientDisconnected(FGuid ClientId)
 		IDisplayClusterScenePreview& PreviewModule = IDisplayClusterScenePreview::Get();
 		for (TClientIdToPerRendererDataMap::ElementType& PerRendererDataPair : *PerRendererDataMap)
 		{
+			EndActorDrag(PerRendererDataPair.Value, ClientId, PerRendererDataPair.Key, true);
 			PreviewModule.DestroyRenderer(PerRendererDataPair.Key);
 		}
 
