[EpicStageApp] API for remote lightcard creation with template/position override/UV support

diff --git a/Engine/Plugins/Runtime/nDisplay/Source/DisplayClusterLightCardEditor/Private/DisplayClusterLightCardEditor.cpp b/Engine/Plugins/Runtime/nDisplay/Source/DisplayClusterLightCardEditor/Private/DisplayClusterLightCardEditor.cpp
--- a/Engine/Plugins/Runtime/nDisplay/Source/DisplayClusterLightCardEditor/Private/DisplayClusterLightCardEditor.cpp
+++ b/Engine/Plugins/Runtime/nDisplay/Source/DisplayClusterLightCardEditor/Private/DisplayClusterLightCardEditor.cpp
@@ -285,6 +285,14 @@ ADisplayClusterLightCardActor* FDisplayClusterLightCardEditor::SpawnLightCardFro
 		AddRecentlyPlacedItem(MoveTemp(RecentlyPlacedItem));
 		SelectActors({NewLightCard});
 	}
+
+#if WITH_EDITOR
+	// Need to call this if spawned from a template since this would normally be called in SpawnActor
+	if (GIsEditor)
+	{
+		GEditor->BroadcastLevelActorAdded(NewLightCard);
+	}
+#endif
 	
 	return NewLightCard;
 }
diff --git a/Engine/Plugins/Runtime/nDisplay/Source/DisplayClusterLightCardEditor/Private/DisplayClusterLightCardEditorModule.cpp b/Engine/Plugins/Runtime/nDisplay/Source/DisplayClusterLightCardEditor/Private/DisplayClusterLightCardEditorModule.cpp
--- a/Engine/Plugins/Runtime/nDisplay/Source/DisplayClusterLightCardEditor/Private/DisplayClusterLightCardEditorModule.cpp
+++ b/Engine/Plugins/Runtime/nDisplay/Source/DisplayClusterLightCardEditor/Private/DisplayClusterLightCardEditorModule.cpp
@@ -8,6 +8,7 @@
 #include "DisplayClusterRootActor.h"
 #include "DisplayClusterLightCardEditor.h"
 #include "DetailCustomizations/DisplayClusterLightCardActorDetails.h"
+#include "LightCardTemplates/DisplayClusterLightCardTemplate.h"
 #include "Settings/DisplayClusterLightCardEditorSettings.h"
 
 #include "IDisplayClusterOperator.h"
@@ -70,6 +71,12 @@ void FDisplayClusterLightCardEditorModule::ShowLabels(const FLabelArgs& InArgs)
 	}
 }
 
+UDisplayClusterLightCardTemplate* FDisplayClusterLightCardEditorModule::GetDefaultLightCardTemplate() const
+{
+	const UDisplayClusterLightCardEditorProjectSettings* Settings = GetDefault<UDisplayClusterLightCardEditorProjectSettings>();
+	return Settings->DefaultLightCardTemplate.LoadSynchronous();
+}
+
 void FDisplayClusterLightCardEditorModule::RegisterSettings()
 {
 	if (ISettingsModule* SettingsModule = FModuleManager::GetModulePtr<ISettingsModule>("Settings"))
diff --git a/Engine/Plugins/Runtime/nDisplay/Source/DisplayClusterLightCardEditor/Private/DisplayClusterLightCardEditorModule.h b/Engine/Plugins/Runtime/nDisplay/Source/DisplayClusterLightCardEditor/Private/DisplayClusterLightCardEditorModule.h
--- a/Engine/Plugins/Runtime/nDisplay/Source/DisplayClusterLightCardEditor/Private/DisplayClusterLightCardEditorModule.h
+++ b/Engine/Plugins/Runtime/nDisplay/Source/DisplayClusterLightCardEditor/Private/DisplayClusterLightCardEditorModule.h
@@ -16,6 +16,7 @@ public:
 	virtual void ShutdownModule() override;
 	
 	virtual void ShowLabels(const FLabelArgs& InArgs) override;
+	virtual UDisplayClusterLightCardTemplate* GetDefaultLightCardTemplate() const override;
 	
 private:
 	void RegisterSettings();
diff --git a/Engine/Plugins/Runtime/nDisplay/Source/DisplayClusterLightCardEditor/Private/LightCardTemplates/DisplayClusterLightCardTemplate.h b/Engine/Plugins/Runtime/nDisplay/Source/DisplayClusterLightCardEditor/Private/LightCardTemplates/DisplayClusterLightCardTemplate.h
deleted file mode 100644
diff --git a/Engine/Plugins/Runtime/nDisplay/Source/DisplayClusterLightCardEditor/Private/LightCardTemplates/DisplayClusterLightCardTemplateHelpers.cpp b/Engine/Plugins/Runtime/nDisplay/Source/DisplayClusterLightCardEditor/Private/LightCardTemplates/DisplayClusterLightCardTemplateHelpers.cpp
--- a/Engine/Plugins/Runtime/nDisplay/Source/DisplayClusterLightCardEditor/Private/LightCardTemplates/DisplayClusterLightCardTemplateHelpers.cpp
+++ b/Engine/Plugins/Runtime/nDisplay/Source/DisplayClusterLightCardEditor/Private/LightCardTemplates/DisplayClusterLightCardTemplateHelpers.cpp
@@ -2,9 +2,9 @@
 
 #include "DisplayClusterLightCardTemplateHelpers.h"
 
-#include "DisplayClusterLightCardTemplate.h"
 #include "AssetRegistry/AssetRegistryModule.h"
 #include "AssetRegistry/IAssetRegistry.h"
+#include "LightCardTemplates/DisplayClusterLightCardTemplate.h"
 
 TArray<UDisplayClusterLightCardTemplate*> UE::DisplayClusterLightCardTemplateHelpers::GetLightCardTemplates(bool bFavoritesOnly)
 {
diff --git a/Engine/Plugins/Runtime/nDisplay/Source/DisplayClusterLightCardEditor/Private/LightCardTemplates/SDisplayClusterLightCardTemplateList.cpp b/Engine/Plugins/Runtime/nDisplay/Source/DisplayClusterLightCardEditor/Private/LightCardTemplates/SDisplayClusterLightCardTemplateList.cpp
--- a/Engine/Plugins/Runtime/nDisplay/Source/DisplayClusterLightCardEditor/Private/LightCardTemplates/SDisplayClusterLightCardTemplateList.cpp
+++ b/Engine/Plugins/Runtime/nDisplay/Source/DisplayClusterLightCardEditor/Private/LightCardTemplates/SDisplayClusterLightCardTemplateList.cpp
@@ -2,10 +2,10 @@
 
 #include "SDisplayClusterLightCardTemplateList.h"
 
-#include "DisplayClusterLightCardTemplate.h"
 #include "DisplayClusterLightCardTemplateDragDropOp.h"
 #include "DisplayClusterLightCardTemplateHelpers.h"
 #include "DisplayClusterLightCardEditor.h"
+#include "LightCardTemplates/DisplayClusterLightCardTemplate.h"
 
 #include "AssetRegistry/AssetRegistryModule.h"
 #include "Engine/Texture.h"
diff --git a/Engine/Plugins/Runtime/nDisplay/Source/DisplayClusterLightCardEditor/Public/IDisplayClusterLightCardEditor.h b/Engine/Plugins/Runtime/nDisplay/Source/DisplayClusterLightCardEditor/Public/IDisplayClusterLightCardEditor.h
--- a/Engine/Plugins/Runtime/nDisplay/Source/DisplayClusterLightCardEditor/Public/IDisplayClusterLightCardEditor.h
+++ b/Engine/Plugins/Runtime/nDisplay/Source/DisplayClusterLightCardEditor/Public/IDisplayClusterLightCardEditor.h
@@ -6,6 +6,8 @@
 #include "Modules/ModuleInterface.h"
 #include "Modules/ModuleManager.h"
 
+class UDisplayClusterLightCardTemplate;
+
 /**
  * Display Cluster Light Card Editor module interface
  */
@@ -64,4 +66,9 @@ public:
 	 * @param InArgs The arguments for displaying labels.
 	 */
 	virtual void ShowLabels(const FLabelArgs& InArgs) = 0;
+
+	/**
+	 * Get the default template to use when creating a new light card
+	 */
+	virtual UDisplayClusterLightCardTemplate* GetDefaultLightCardTemplate() const = 0;
 };
\ No newline at end of file
diff --git a/Engine/Plugins/Runtime/nDisplay/Source/DisplayClusterLightCardEditor/Public/LightCardTemplates/DisplayClusterLightCardTemplate.h b/Engine/Plugins/Runtime/nDisplay/Source/DisplayClusterLightCardEditor/Public/LightCardTemplates/DisplayClusterLightCardTemplate.h
new file mode 100644
--- /dev/null
+++ b/Engine/Plugins/Runtime/nDisplay/Source/DisplayClusterLightCardEditor/Public/LightCardTemplates/DisplayClusterLightCardTemplate.h
@@ -0,0 +1,25 @@
+﻿// Copyright Epic Games, Inc. All Rights Reserved.
+
+#pragma once
+
+#include "DisplayClusterLightCardActor.h"
+
+#include "DisplayClusterLightCardTemplate.generated.h"
+
+/**
+ * A template asset to store appearance settings from Light Card actors.
+ */
+UCLASS(NotBlueprintType, NotBlueprintable, NotPlaceable, PerObjectConfig, MinimalAPI, config=EditorPerProjectUserSettings)
+class UDisplayClusterLightCardTemplate : public UObject
+{
+	GENERATED_BODY()
+
+public:
+	/** The instanced template object containing user settings for the light card. */
+	UPROPERTY(VisibleAnywhere, Instanced, Category = Template, NoClear, meta = (ShowOnlyInnerProperties))
+	TObjectPtr<ADisplayClusterLightCardActor> LightCardActor;
+
+	/** If the user has marked this a favorite template. */
+	UPROPERTY(Config)
+	bool bIsFavorite;
+};
diff --git a/Engine/Plugins/Runtime/nDisplay/Source/DisplayClusterScenePreview/Private/DisplayClusterLightCardEditorHelper.cpp b/Engine/Plugins/Runtime/nDisplay/Source/DisplayClusterScenePreview/Private/DisplayClusterLightCardEditorHelper.cpp
--- a/Engine/Plugins/Runtime/nDisplay/Source/DisplayClusterScenePreview/Private/DisplayClusterLightCardEditorHelper.cpp
+++ b/Engine/Plugins/Runtime/nDisplay/Source/DisplayClusterScenePreview/Private/DisplayClusterLightCardEditorHelper.cpp
@@ -404,6 +404,8 @@ const UTexture2D* FDisplayClusterLightCardEditorHelper::GetNormalMapTexture(bool
 void FDisplayClusterLightCardEditorHelper::MoveActorsToPixel(
 	const TArray<FDisplayClusterWeakStageActorPtr>& Actors, const FIntPoint& PixelPos, const FSceneView& SceneView)
 {
+	UpdateProjectionOriginComponent();
+
 	FVector Origin;
 	FVector Direction;
 	CalculateOriginAndDirectionFromPixelPosition(PixelPos, SceneView, FVector::ZeroVector, Origin, Direction);
diff --git a/Engine/Plugins/VirtualProduction/EpicStageApp/Source/EpicStageApp/EpicStageApp.build.cs b/Engine/Plugins/VirtualProduction/EpicStageApp/Source/EpicStageApp/EpicStageApp.build.cs
--- a/Engine/Plugins/VirtualProduction/EpicStageApp/Source/EpicStageApp/EpicStageApp.build.cs
+++ b/Engine/Plugins/VirtualProduction/EpicStageApp/Source/EpicStageApp/EpicStageApp.build.cs
@@ -36,6 +36,7 @@ public class EpicStageApp : ModuleRules
 		{
 			PrivateDependencyModuleNames.AddRange(
 				new string[] {
+					"DisplayClusterLightCardEditor",
 					"UnrealEd",
 				}
 			);
diff --git a/Engine/Plugins/VirtualProduction/EpicStageApp/Source/EpicStageApp/Private/StageAppRequest.h b/Engine/Plugins/VirtualProduction/EpicStageApp/Source/EpicStageApp/Private/StageAppRequest.h
--- a/Engine/Plugins/VirtualProduction/EpicStageApp/Source/EpicStageApp/Private/StageAppRequest.h
+++ b/Engine/Plugins/VirtualProduction/EpicStageApp/Source/EpicStageApp/Private/StageAppRequest.h
@@ -327,3 +327,56 @@ struct FRCWebSocketNDisplayPreviewActorDragEndBody : public FRCRequest
 	UPROPERTY()
 	int64 SequenceNumber = -1;
 };
+
+/**
+ * Holds a request made via websocket to create a lightcard, optionally positioned relative to the previewed area.
+ */
+USTRUCT()
+struct FRCWebSocketNDisplayPreviewLightCardCreateBody : public FRCRequest
+{
+	GENERATED_BODY()
+
+	FRCWebSocketNDisplayPreviewLightCardCreateBody()
+	{
+		AddStructParameter(ParametersFieldLabel());
+	}
+
+	/**
+	 * Get the label for the property value struct.
+	 */
+	static FString ParametersFieldLabel() { return TEXT("Parameters"); }
+
+	/**
+	 * The ID of the preview renderer returned when it was created.
+	 */
+	UPROPERTY()
+	int32 RendererId = -1;
+
+	/**
+	 * The path of the template to use for the lightcard. If empty, a lightcard will be created with default settings.
+	 */
+	UPROPERTY()
+	FString TemplatePath = "";
+
+	/**
+	 * If true, override the default/template position for the lightcard with the provided Position.
+	 * Otherwise, use the template's position, or if no template, create at viewport center (0.5, 0.5).
+	 */
+	UPROPERTY()
+	bool OverridePosition = false;
+
+	/**
+	 * If OverridePosition is true and PreviewRendererId points to a valid preview renderer, place the
+	 * lightcard at this normalized viewport coordinate (in range [0, 1]).
+	 */
+	UPROPERTY()
+	FVector2D Position = FVector2D::ZeroVector;
+
+	/**
+	 * An optional number that will be passed back in the RequestedActorsCreated response to tell apart
+	 * the results of multiple requests.
+	 */
+	UPROPERTY()
+	int32 RequestId = -1;
+};
+
diff --git a/Engine/Plugins/VirtualProduction/EpicStageApp/Source/EpicStageApp/Private/StageAppResponse.h b/Engine/Plugins/VirtualProduction/EpicStageApp/Source/EpicStageApp/Private/StageAppResponse.h
--- a/Engine/Plugins/VirtualProduction/EpicStageApp/Source/EpicStageApp/Private/StageAppResponse.h
+++ b/Engine/Plugins/VirtualProduction/EpicStageApp/Source/EpicStageApp/Private/StageAppResponse.h
@@ -127,3 +127,34 @@ struct FRCActorDragCancelled
 	int32 RendererId = -1;
 };
 
+/**
+ * Event sent to a client when light cards were created at the request of the client.
+ */
+USTRUCT()
+struct FRCRequestedActorsCreated
+{
+	GENERATED_BODY()
+
+	FRCRequestedActorsCreated()
+	: Type(TEXT("RequestedActorsCreated"))
+	{
+	}
+
+	/**
+	 * Type of the event.
+	 */
+	UPROPERTY()
+	FString Type;
+
+	/**
+	 * The list of newly-created actors' paths.
+	 */
+	UPROPERTY()
+	TArray<FString> ActorPaths;
+
+	/**
+	 * The ID passed when the actors were created, or -1 if no ID was provided.
+	 */
+	UPROPERTY()
+	int32 RequestId = -1;
+};
diff --git a/Engine/Plugins/VirtualProduction/EpicStageApp/Source/EpicStageApp/Private/StageAppRouteHandler.cpp b/Engine/Plugins/VirtualProduction/EpicStageApp/Source/EpicStageApp/Private/StageAppRouteHandler.cpp
--- a/Engine/Plugins/VirtualProduction/EpicStageApp/Source/EpicStageApp/Private/StageAppRouteHandler.cpp
+++ b/Engine/Plugins/VirtualProduction/EpicStageApp/Source/EpicStageApp/Private/StageAppRouteHandler.cpp
@@ -4,12 +4,14 @@
 
 #include "CanvasTypes.h"
 #include "Components/DisplayClusterCameraComponent.h"
+#include "DisplayClusterConfigurationTypes.h"
 #include "DisplayClusterLightCardActor.h"
 #include "DisplayClusterLightCardEditorHelper.h"
 #include "DisplayClusterRootActor.h"
 #include "Engine/TextureRenderTarget2D.h"
 #include "IDisplayClusterScenePreview.h"
 #include "IImageWrapperModule.h"
+#include "IRemoteControlModule.h"
 #include "Misc/Base64.h"
 #include "Modules/ModuleManager.h"
 #include "RemoteControlWebsocketRoute.h"
@@ -20,6 +22,9 @@
 #if WITH_EDITOR
 #include "Editor.h"
 #include "Editor/TransBuffer.h"
+#include "IDisplayClusterLightCardEditor.h"
+#include "LightCardTemplates/DisplayClusterLightCardTemplate.h"
+#include "ScopedTransaction.h"
 #endif
 
 #define LOCTEXT_NAMESPACE "StageAppRouteHandler"
@@ -124,7 +129,7 @@ void FStageAppRouteHandler::FPerRendererData::SetPreviewSettings(const FRCWebSoc
 	UpdateLightCardHelperSettings();
 }
 
-bool FStageAppRouteHandler::FPerRendererData::GetSceneViewInitOptions(FSceneViewInitOptions& ViewInitOptions)
+bool FStageAppRouteHandler::FPerRendererData::GetSceneViewInitOptions(FSceneViewInitOptions& ViewInitOptions, bool bApplyActorRotation)
 {
 	ADisplayClusterRootActor* RootActor = GetRootActor();
 	if (!RootActor)
@@ -141,7 +146,23 @@ bool FStageAppRouteHandler::FPerRendererData::GetSceneViewInitOptions(FSceneView
 	// In UV mode, lock the location and rotation
 	const bool bIsUVMode = GetLightCardHelper().GetProjectionMode() == EDisplayClusterMeshProjectionType::UV;
 	const FVector Location = bIsUVMode ? FVector::ZeroVector : ViewOriginComponent->GetComponentLocation();
-	const FRotator Rotation = bIsUVMode ? FVector::ForwardVector.Rotation() : FRotator(RootActor->GetActorRotation().Quaternion() * PreviewSettings.Rotation.Quaternion());
+	FRotator Rotation;
+	
+	if (bIsUVMode)
+	{
+		Rotation = FVector::ForwardVector.Rotation();
+	}
+	else
+	{
+		if (bApplyActorRotation)
+		{
+			Rotation = FRotator(RootActor->GetActorRotation().Quaternion() * PreviewSettings.Rotation.Quaternion());
+		}
+		else
+		{
+			Rotation = PreviewSettings.Rotation;
+		}
+	}
 
 	GetLightCardHelper().GetSceneViewInitOptions(
 		ViewInitOptions,
@@ -252,6 +273,12 @@ void FStageAppRouteHandler::RegisterRoutes(IWebRemoteControlModule& WebRemoteCon
 		TEXT("ndisplay.preview.actor.drag.end"),
 		FWebSocketMessageDelegate::CreateRaw(this, &FStageAppRouteHandler::HandleWebSocketNDisplayPreviewActorDragEnd)
 	));
+
+	RegisterRoute(MakeUnique<FRemoteControlWebsocketRoute>(
+		TEXT("Create a lightcard positioned relative to the preview renderer's viewport"),
+		TEXT("ndisplay.preview.lightcard.create"),
+		FWebSocketMessageDelegate::CreateRaw(this, &FStageAppRouteHandler::HandleWebSocketNDisplayPreviewLightCardCreate)
+	));
 }
 
 void FStageAppRouteHandler::RegisterEngineEvents()
@@ -701,6 +728,157 @@ void FStageAppRouteHandler::HandleWebSocketNDisplayPreviewActorDragEnd(const FRe
 	EndActorDrag(*PerRendererData, WebSocketMessage.ClientId, Body.RendererId, true);
 }
 
+void FStageAppRouteHandler::HandleWebSocketNDisplayPreviewLightCardCreate(const FRemoteControlWebSocketMessage& WebSocketMessage)
+{
+	FRCWebSocketNDisplayPreviewLightCardCreateBody Body;
+	if (!WebRemoteControlUtils::DeserializeMessage(WebSocketMessage.RequestPayload, Body))
+	{
+		return;
+	}
+
+	FPerRendererData* PerRendererData = GetClientPerRendererData(WebSocketMessage.ClientId, Body.RendererId);
+	if (!PerRendererData)
+	{
+		return;
+	}
+
+	ADisplayClusterRootActor* RootActor = PerRendererData->GetRootActor();
+	ULevel* Level = RootActor->GetWorld()->GetCurrentLevel();
+	if (!Level)
+	{
+		return;
+	}
+
+	ADisplayClusterLightCardActor* NewLightCard;
+	const bool bIsInUVMode = PerRendererData->GetPreviewSettings().ProjectionType == ERCWebSocketNDisplayPreviewRenderProjectionType::UV;
+
+#if WITH_EDITOR
+	const UDisplayClusterLightCardTemplate* Template;
+
+	if (!Body.TemplatePath.IsEmpty())
+	{
+		Template = Cast<UDisplayClusterLightCardTemplate>(StaticLoadObject(UDisplayClusterLightCardTemplate::StaticClass(), nullptr, *Body.TemplatePath));
+	}
+	else
+	{
+		// Use the default template if available
+		Template = IDisplayClusterLightCardEditor::Get().GetDefaultLightCardTemplate();
+	}
+
+	if (!bIsInUVMode && Template && Template->LightCardActor->bIsUVLightCard)
+	{
+		// Can't spawn a UV lightcard in non-UV mode
+		return;
+	}
+
+	FScopedTransaction Transaction(LOCTEXT("CreateLightCard", "Create Light Card"));
+
+	if (Template)
+	{
+		// Duplicate the lightcard from its template
+		FName UniqueName = *Template->GetName().Replace(TEXT("Template"), TEXT(""));
+		if (StaticFindObjectFast(Template->LightCardActor->GetClass(), Level, UniqueName))
+		{
+			UniqueName = MakeUniqueObjectName(Level, Template->LightCardActor->GetClass(), UniqueName);
+		}
+
+		// Duplicate, don't copy properties or spawn from a template. Doing so will copy component data incorrectly,
+		// specifically the static mesh override textures. They will be parented to the template, not the level instance
+		// and prevent the map from saving.
+		NewLightCard = CastChecked<ADisplayClusterLightCardActor>(StaticDuplicateObject(Template->LightCardActor.Get(), Level, UniqueName));
+		Level->AddLoadedActor(NewLightCard);
+
+		NewLightCard->SetActorLabel(NewLightCard->GetName());
+	}
+	else
+#endif
+	{
+		// Create the light card
+		const FVector SpawnLocation = RootActor->GetDefaultCamera()->GetComponentLocation();
+		FRotator SpawnRotation = RootActor->GetDefaultCamera()->GetComponentRotation();
+		SpawnRotation.Yaw -= 180.f;
+
+		FActorSpawnParameters SpawnParameters;
+		SpawnParameters.bNoFail = true;
+		SpawnParameters.SpawnCollisionHandlingOverride = ESpawnActorCollisionHandlingMethod::AdjustIfPossibleButAlwaysSpawn;
+		SpawnParameters.Name = TEXT("LightCard");
+		SpawnParameters.NameMode = FActorSpawnParameters::ESpawnActorNameMode::Requested;
+		SpawnParameters.OverrideLevel = RootActor->GetLevel();
+
+		NewLightCard = CastChecked<ADisplayClusterLightCardActor>(
+			RootActor->GetWorld()->SpawnActor(ADisplayClusterLightCardActor::StaticClass(),
+				&SpawnLocation, &SpawnRotation, MoveTemp(SpawnParameters)));
+
+		NewLightCard->Color = FLinearColor::Gray;
+
+		// If no position was specified, spawn at center of the viewport
+		if (!Body.OverridePosition)
+		{
+			Body.OverridePosition = true;
+			Body.Position = FVector2D(0.5, 0.5);
+		}
+
+		if (bIsInUVMode)
+		{
+			NewLightCard->Feathering = 0.05f; // Just enough to avoid jagged look on UV lightcards.
+		}
+	}
+
+	// Convert a non-UV LC (or brand new LC) to a UV LC
+	if (bIsInUVMode && !NewLightCard->bIsUVLightCard)
+	{
+		NewLightCard->bIsUVLightCard = true;
+		NewLightCard->Scale /= 4;
+	}
+
+	// Override the location
+	if (Body.OverridePosition)
+	{
+		const FIntPoint& Resolution = PerRendererData->GetPreviewSettings().Resolution;
+
+		const FIntPoint PixelPos(
+			FMath::RoundToInt(Body.Position.X * Resolution.X),
+			FMath::RoundToInt(Body.Position.Y * Resolution.Y)
+		);
+
+		FDisplayClusterLightCardEditorHelper& LightCardHelper = PerRendererData->GetLightCardHelper();
+
+		FSceneViewInitOptions ViewInitOptions;
+		if (PerRendererData->GetSceneViewInitOptions(ViewInitOptions, false))
+		{
+			const FSceneView View(ViewInitOptions);
+			LightCardHelper.MoveActorsToPixel({ NewLightCard }, PixelPos, View);
+		}
+	}
+
+#if WITH_EDITOR
+	NewLightCard->SetActorLabel(NewLightCard->GetName());
+#endif
+
+	// Add it to the root actor
+	NewLightCard->AddToLightCardLayer(RootActor);
+
+#if WITH_EDITOR
+	// Required so operator panel updates
+	RootActor->PostEditChange();
+
+	// Need to call this if spawned from a template since this would normally be called in SpawnActor
+	if (GIsEditor && Template)
+	{
+		GEditor->BroadcastLevelActorAdded(NewLightCard);
+	}
+#endif
+
+	// Send new actor back to the client
+	FRCRequestedActorsCreated Event;
+	Event.RequestId = Body.RequestId;
+	Event.ActorPaths.Add(NewLightCard->GetPathName());
+
+	TArray<uint8> Payload;
+	WebRemoteControlUtils::SerializeMessage(Event, Payload);
+	RemoteControlModule->SendWebsocketMessage(WebSocketMessage.ClientId, Payload);
+}
+
 void FStageAppRouteHandler::HandleClientDisconnected(FGuid ClientId)
 {
 	// Destroy the client's renderers
diff --git a/Engine/Plugins/VirtualProduction/EpicStageApp/Source/EpicStageApp/Private/StageAppRouteHandler.h b/Engine/Plugins/VirtualProduction/EpicStageApp/Source/EpicStageApp/Private/StageAppRouteHandler.h
--- a/Engine/Plugins/VirtualProduction/EpicStageApp/Source/EpicStageApp/Private/StageAppRouteHandler.h
+++ b/Engine/Plugins/VirtualProduction/EpicStageApp/Source/EpicStageApp/Private/StageAppRouteHandler.h
@@ -54,7 +54,7 @@ private:
 		void SetPreviewSettings(const FRCWebSocketNDisplayPreviewRendererSettings& NewSettings);
 
 		/** Get the scene view init options to be used for preview renders based on the stored settings. */
-		bool GetSceneViewInitOptions(FSceneViewInitOptions& ViewInitOptions);
+		bool GetSceneViewInitOptions(FSceneViewInitOptions& ViewInitOptions, bool bApplyActorRotation = true);
 
 		/** Update the sequence number if the provided one is higher than the previous. */
 		void UpdateDragSequenceNumber(int64 ReceivedSequenceNumber);
@@ -127,6 +127,9 @@ private:
 	/** Handles ending an actor drag operation. */
 	void HandleWebSocketNDisplayPreviewActorDragEnd(const FRemoteControlWebSocketMessage& WebSocketMessage);
 
+	/** Handles creating a lightcard. */
+	void HandleWebSocketNDisplayPreviewLightCardCreate(const FRemoteControlWebSocketMessage& WebSocketMessage);
+
 	/** Called when a client disconnects from the WebSocket server. */
 	void HandleClientDisconnected(FGuid ClientId);
 
