[EpicStageApp] Use shared function for creating lightcards to avoid code duplication

diff --git a/Engine/Plugins/Runtime/nDisplay/Source/DisplayClusterLightCardEditor/Private/DisplayClusterLightCardEditorModule.cpp b/Engine/Plugins/Runtime/nDisplay/Source/DisplayClusterLightCardEditor/Private/DisplayClusterLightCardEditorModule.cpp
--- a/Engine/Plugins/Runtime/nDisplay/Source/DisplayClusterLightCardEditor/Private/DisplayClusterLightCardEditorModule.cpp
+++ b/Engine/Plugins/Runtime/nDisplay/Source/DisplayClusterLightCardEditor/Private/DisplayClusterLightCardEditorModule.cpp
@@ -71,7 +71,7 @@ void FDisplayClusterLightCardEditorModule::ShowLabels(const FLabelArgs& InArgs)
 	}
 }
 
-UDisplayClusterLightCardTemplate* FDisplayClusterLightCardEditorModule::GetDefaultLightCardTemplate() const
+UDisplayClusterStageActorTemplate* FDisplayClusterLightCardEditorModule::GetDefaultLightCardTemplate() const
 {
 	const UDisplayClusterLightCardEditorProjectSettings* Settings = GetDefault<UDisplayClusterLightCardEditorProjectSettings>();
 	return Settings->DefaultLightCardTemplate.LoadSynchronous();
diff --git a/Engine/Plugins/Runtime/nDisplay/Source/DisplayClusterLightCardEditor/Private/DisplayClusterLightCardEditorModule.h b/Engine/Plugins/Runtime/nDisplay/Source/DisplayClusterLightCardEditor/Private/DisplayClusterLightCardEditorModule.h
--- a/Engine/Plugins/Runtime/nDisplay/Source/DisplayClusterLightCardEditor/Private/DisplayClusterLightCardEditorModule.h
+++ b/Engine/Plugins/Runtime/nDisplay/Source/DisplayClusterLightCardEditor/Private/DisplayClusterLightCardEditorModule.h
@@ -16,7 +16,7 @@ public:
 	virtual void ShutdownModule() override;
 	
 	virtual void ShowLabels(const FLabelArgs& InArgs) override;
-	virtual UDisplayClusterLightCardTemplate* GetDefaultLightCardTemplate() const override;
+	virtual UDisplayClusterStageActorTemplate* GetDefaultLightCardTemplate() const override;
 	
 private:
 	void RegisterSettings();
diff --git a/Engine/Plugins/Runtime/nDisplay/Source/DisplayClusterLightCardEditor/Private/LightCardTemplates/DisplayClusterLightCardTemplate.h b/Engine/Plugins/Runtime/nDisplay/Source/DisplayClusterLightCardEditor/Private/LightCardTemplates/DisplayClusterLightCardTemplate.h
new file mode 100644
--- /dev/null
+++ b/Engine/Plugins/Runtime/nDisplay/Source/DisplayClusterLightCardEditor/Private/LightCardTemplates/DisplayClusterLightCardTemplate.h
@@ -0,0 +1,27 @@
+﻿// Copyright Epic Games, Inc. All Rights Reserved.
+
+#pragma once
+
+#include "DisplayClusterLightCardActor.h"
+#include "StageActor/DisplayClusterStageActorTemplate.h"
+
+#include "DisplayClusterLightCardTemplate.generated.h"
+
+/**
+ * A template asset to store appearance settings from Light Card actors.
+ */
+UCLASS(NotBlueprintType, NotBlueprintable, NotPlaceable)
+class UDisplayClusterLightCardTemplate : public UDisplayClusterStageActorTemplate
+{
+	GENERATED_BODY()
+
+public:
+	// UDisplayClusterStageActorTemplate
+	virtual AActor* GetTemplateActor() const override { return LightCardActor; }
+	// ~UDisplayClusterStageActorTemplate
+	
+public:
+	/** The instanced template object containing user settings for the light card. */
+	UPROPERTY(VisibleAnywhere, Instanced, Category = Template, NoClear, meta = (ShowOnlyInnerProperties))
+	TObjectPtr<ADisplayClusterLightCardActor> LightCardActor;
+};
diff --git a/Engine/Plugins/Runtime/nDisplay/Source/DisplayClusterLightCardEditor/Public/IDisplayClusterLightCardEditor.h b/Engine/Plugins/Runtime/nDisplay/Source/DisplayClusterLightCardEditor/Public/IDisplayClusterLightCardEditor.h
--- a/Engine/Plugins/Runtime/nDisplay/Source/DisplayClusterLightCardEditor/Public/IDisplayClusterLightCardEditor.h
+++ b/Engine/Plugins/Runtime/nDisplay/Source/DisplayClusterLightCardEditor/Public/IDisplayClusterLightCardEditor.h
@@ -6,7 +6,7 @@
 #include "Modules/ModuleInterface.h"
 #include "Modules/ModuleManager.h"
 
-class UDisplayClusterLightCardTemplate;
+class UDisplayClusterStageActorTemplate;
 
 /**
  * Display Cluster Light Card Editor module interface
@@ -70,5 +70,5 @@ public:
 	/**
 	 * Get the default template to use when creating a new light card
 	 */
-	virtual UDisplayClusterLightCardTemplate* GetDefaultLightCardTemplate() const = 0;
+	virtual UDisplayClusterStageActorTemplate* GetDefaultLightCardTemplate() const = 0;
 };
\ No newline at end of file
diff --git a/Engine/Plugins/Runtime/nDisplay/Source/DisplayClusterLightCardEditor/Public/LightCardTemplates/DisplayClusterLightCardTemplate.h b/Engine/Plugins/Runtime/nDisplay/Source/DisplayClusterLightCardEditor/Public/LightCardTemplates/DisplayClusterLightCardTemplate.h
deleted file mode 100644
diff --git a/Engine/Plugins/VirtualProduction/EpicStageApp/Source/EpicStageApp/Private/StageAppRequest.h b/Engine/Plugins/VirtualProduction/EpicStageApp/Source/EpicStageApp/Private/StageAppRequest.h
--- a/Engine/Plugins/VirtualProduction/EpicStageApp/Source/EpicStageApp/Private/StageAppRequest.h
+++ b/Engine/Plugins/VirtualProduction/EpicStageApp/Source/EpicStageApp/Private/StageAppRequest.h
@@ -329,14 +329,14 @@ struct FRCWebSocketNDisplayPreviewActorDragEndBody : public FRCRequest
 };
 
 /**
- * Holds a request made via websocket to create a lightcard, optionally positioned relative to the previewed area.
+ * Holds a request made via websocket to create an actor, optionally positioned relative to the previewed area.
  */
 USTRUCT()
-struct FRCWebSocketNDisplayPreviewLightCardCreateBody : public FRCRequest
+struct FRCWebSocketNDisplayPreviewActorCreateBody : public FRCRequest
 {
 	GENERATED_BODY()
 
-	FRCWebSocketNDisplayPreviewLightCardCreateBody()
+	FRCWebSocketNDisplayPreviewActorCreateBody()
 	{
 		AddStructParameter(ParametersFieldLabel());
 	}
@@ -347,11 +347,23 @@ struct FRCWebSocketNDisplayPreviewLightCardCreateBody : public FRCRequest
 	static FString ParametersFieldLabel() { return TEXT("Parameters"); }
 
 	/**
-	 * The ID of the preview renderer returned when it was created.
+	 * The ID of the preview renderer to use when positioning this light card.
 	 */
 	UPROPERTY()
 	int32 RendererId = -1;
 
+	/**
+	 * The name of the actor to spawn. A number will be appended if this conflicts with another name.
+	 */
+	UPROPERTY()
+	FString ActorName = "";
+
+	/**
+	 * The name of the class of the actor to spawn.
+	 */
+	UPROPERTY()
+	FString ActorClass = "";
+
 	/**
 	 * The path of the template to use for the lightcard. If empty, a lightcard will be created with default settings.
 	 */
diff --git a/Engine/Plugins/VirtualProduction/EpicStageApp/Source/EpicStageApp/Private/StageAppRouteHandler.cpp b/Engine/Plugins/VirtualProduction/EpicStageApp/Source/EpicStageApp/Private/StageAppRouteHandler.cpp
--- a/Engine/Plugins/VirtualProduction/EpicStageApp/Source/EpicStageApp/Private/StageAppRouteHandler.cpp
+++ b/Engine/Plugins/VirtualProduction/EpicStageApp/Source/EpicStageApp/Private/StageAppRouteHandler.cpp
@@ -17,13 +17,13 @@
 #include "RemoteControlWebsocketRoute.h"
 #include "StageAppResponse.h"
 #include "WebRemoteControlUtils.h"
+#include "StageActor/DisplayClusterStageActorTemplate.h"
 #include "StageActor/DisplayClusterWeakStageActorPtr.h"
 
 #if WITH_EDITOR
 #include "Editor.h"
 #include "Editor/TransBuffer.h"
 #include "IDisplayClusterLightCardEditor.h"
-#include "LightCardTemplates/DisplayClusterLightCardTemplate.h"
 #include "ScopedTransaction.h"
 #endif
 
@@ -275,9 +275,9 @@ void FStageAppRouteHandler::RegisterRoutes(IWebRemoteControlModule& WebRemoteCon
 	));
 
 	RegisterRoute(MakeUnique<FRemoteControlWebsocketRoute>(
-		TEXT("Create a lightcard positioned relative to the preview renderer's viewport"),
-		TEXT("ndisplay.preview.lightcard.create"),
-		FWebSocketMessageDelegate::CreateRaw(this, &FStageAppRouteHandler::HandleWebSocketNDisplayPreviewLightCardCreate)
+		TEXT("Create an actor positioned relative to the preview renderer's viewport"),
+		TEXT("ndisplay.preview.actor.create"),
+		FWebSocketMessageDelegate::CreateRaw(this, &FStageAppRouteHandler::HandleWebSocketNDisplayPreviewActorCreate)
 	));
 }
 
@@ -728,9 +728,9 @@ void FStageAppRouteHandler::HandleWebSocketNDisplayPreviewActorDragEnd(const FRe
 	EndActorDrag(*PerRendererData, WebSocketMessage.ClientId, Body.RendererId, true);
 }
 
-void FStageAppRouteHandler::HandleWebSocketNDisplayPreviewLightCardCreate(const FRemoteControlWebSocketMessage& WebSocketMessage)
+void FStageAppRouteHandler::HandleWebSocketNDisplayPreviewActorCreate(const FRemoteControlWebSocketMessage& WebSocketMessage)
 {
-	FRCWebSocketNDisplayPreviewLightCardCreateBody Body;
+	FRCWebSocketNDisplayPreviewActorCreateBody Body;
 	if (!WebRemoteControlUtils::DeserializeMessage(WebSocketMessage.RequestPayload, Body))
 	{
 		return;
@@ -743,96 +743,55 @@ void FStageAppRouteHandler::HandleWebSocketNDisplayPreviewLightCardCreate(const
 	}
 
 	ADisplayClusterRootActor* RootActor = PerRendererData->GetRootActor();
-	ULevel* Level = RootActor->GetWorld()->GetCurrentLevel();
-	if (!Level)
+	if (!RootActor || !RootActor->GetWorld())
 	{
 		return;
 	}
 
-	ADisplayClusterLightCardActor* NewLightCard;
-	const bool bIsInUVMode = PerRendererData->GetPreviewSettings().ProjectionType == ERCWebSocketNDisplayPreviewRenderProjectionType::UV;
-
-#if WITH_EDITOR
-	const UDisplayClusterLightCardTemplate* Template;
-
-	if (!Body.TemplatePath.IsEmpty())
-	{
-		Template = Cast<UDisplayClusterLightCardTemplate>(StaticLoadObject(UDisplayClusterLightCardTemplate::StaticClass(), nullptr, *Body.TemplatePath));
-	}
-	else
+	ULevel* Level = RootActor->GetWorld()->GetCurrentLevel();
+	if (!Level)
 	{
-		// Use the default template if available
-		Template = IDisplayClusterLightCardEditor::Get().GetDefaultLightCardTemplate();
+		return;
 	}
 
-	if (!bIsInUVMode && Template && Template->LightCardActor->bIsUVLightCard)
+	UClass* ActorClass = StaticLoadClass(UObject::StaticClass(), nullptr, *Body.ActorClass);
+	if (!ActorClass)
 	{
-		// Can't spawn a UV lightcard in non-UV mode
 		return;
 	}
 
-	FScopedTransaction Transaction(LOCTEXT("CreateLightCard", "Create Light Card"));
+	const UDisplayClusterStageActorTemplate* Template = nullptr;
 
-	if (Template)
+	if (!Body.TemplatePath.IsEmpty())
 	{
-		// Duplicate the lightcard from its template
-		FName UniqueName = *Template->GetName().Replace(TEXT("Template"), TEXT(""));
-		if (StaticFindObjectFast(Template->LightCardActor->GetClass(), Level, UniqueName))
-		{
-			UniqueName = MakeUniqueObjectName(Level, Template->LightCardActor->GetClass(), UniqueName);
-		}
-
-		// Duplicate, don't copy properties or spawn from a template. Doing so will copy component data incorrectly,
-		// specifically the static mesh override textures. They will be parented to the template, not the level instance
-		// and prevent the map from saving.
-		NewLightCard = CastChecked<ADisplayClusterLightCardActor>(StaticDuplicateObject(Template->LightCardActor.Get(), Level, UniqueName));
-		Level->AddLoadedActor(NewLightCard);
-
-		NewLightCard->SetActorLabel(NewLightCard->GetName());
+		Template = Cast<UDisplayClusterStageActorTemplate>(StaticLoadObject(UDisplayClusterStageActorTemplate::StaticClass(), nullptr, *Body.TemplatePath));
 	}
-	else
-#endif
+#if WITH_EDITOR
+	else if (ActorClass == ADisplayClusterLightCardActor::StaticClass())
 	{
-		// Create the light card
-		const FVector SpawnLocation = RootActor->GetDefaultCamera()->GetComponentLocation();
-		FRotator SpawnRotation = RootActor->GetDefaultCamera()->GetComponentRotation();
-		SpawnRotation.Yaw -= 180.f;
-
-		FActorSpawnParameters SpawnParameters;
-		SpawnParameters.bNoFail = true;
-		SpawnParameters.SpawnCollisionHandlingOverride = ESpawnActorCollisionHandlingMethod::AdjustIfPossibleButAlwaysSpawn;
-		SpawnParameters.Name = TEXT("LightCard");
-		SpawnParameters.NameMode = FActorSpawnParameters::ESpawnActorNameMode::Requested;
-		SpawnParameters.OverrideLevel = RootActor->GetLevel();
-
-		NewLightCard = CastChecked<ADisplayClusterLightCardActor>(
-			RootActor->GetWorld()->SpawnActor(ADisplayClusterLightCardActor::StaticClass(),
-				&SpawnLocation, &SpawnRotation, MoveTemp(SpawnParameters)));
-
-		NewLightCard->Color = FLinearColor::Gray;
-
-		// If no position was specified, spawn at center of the viewport
-		if (!Body.OverridePosition)
-		{
-			Body.OverridePosition = true;
-			Body.Position = FVector2D(0.5, 0.5);
-		}
+		// Use the default lightcard template if available
+		Template = IDisplayClusterLightCardEditor::Get().GetDefaultLightCardTemplate();
+	}
+#endif
 
-		if (bIsInUVMode)
-		{
-			NewLightCard->Feathering = 0.05f; // Just enough to avoid jagged look on UV lightcards.
-		}
+	FDisplayClusterLightCardEditorHelper::FSpawnActorArgs SpawnArgs;
+	{
+		SpawnArgs.ActorClass = ActorClass;
+		SpawnArgs.ActorName = FName(Body.ActorName);
+		SpawnArgs.RootActor = RootActor;
+		SpawnArgs.Template = Template;
+		SpawnArgs.Level = Level;
+		SpawnArgs.ProjectionMode = PerRendererData->GetLightCardHelper().GetProjectionMode();
 	}
 
-	// Convert a non-UV LC (or brand new LC) to a UV LC
-	if (bIsInUVMode && !NewLightCard->bIsUVLightCard)
+	AActor* NewActor = FDisplayClusterLightCardEditorHelper::SpawnStageActor(SpawnArgs);
+	if (!NewActor)
 	{
-		NewLightCard->bIsUVLightCard = true;
-		NewLightCard->Scale /= 4;
+		return;
 	}
 
 	// Override the location
-	if (Body.OverridePosition)
+	if (Body.OverridePosition && NewActor->Implements<UDisplayClusterStageActor>())
 	{
 		const FIntPoint& Resolution = PerRendererData->GetPreviewSettings().Resolution;
 
@@ -847,32 +806,14 @@ void FStageAppRouteHandler::HandleWebSocketNDisplayPreviewLightCardCreate(const
 		if (PerRendererData->GetSceneViewInitOptions(ViewInitOptions, false))
 		{
 			const FSceneView View(ViewInitOptions);
-			LightCardHelper.MoveActorsToPixel({ NewLightCard }, PixelPos, View);
+			LightCardHelper.MoveActorsToPixel({ NewActor }, PixelPos, View);
 		}
 	}
 
-#if WITH_EDITOR
-	NewLightCard->SetActorLabel(NewLightCard->GetName());
-#endif
-
-	// Add it to the root actor
-	NewLightCard->AddToLightCardLayer(RootActor);
-
-#if WITH_EDITOR
-	// Required so operator panel updates
-	RootActor->PostEditChange();
-
-	// Need to call this if spawned from a template since this would normally be called in SpawnActor
-	if (GIsEditor && Template)
-	{
-		GEditor->BroadcastLevelActorAdded(NewLightCard);
-	}
-#endif
-
 	// Send new actor back to the client
 	FRCRequestedActorsCreated Event;
 	Event.RequestId = Body.RequestId;
-	Event.ActorPaths.Add(NewLightCard->GetPathName());
+	Event.ActorPaths.Add(NewActor->GetPathName());
 
 	TArray<uint8> Payload;
 	WebRemoteControlUtils::SerializeMessage(Event, Payload);
diff --git a/Engine/Plugins/VirtualProduction/EpicStageApp/Source/EpicStageApp/Private/StageAppRouteHandler.h b/Engine/Plugins/VirtualProduction/EpicStageApp/Source/EpicStageApp/Private/StageAppRouteHandler.h
--- a/Engine/Plugins/VirtualProduction/EpicStageApp/Source/EpicStageApp/Private/StageAppRouteHandler.h
+++ b/Engine/Plugins/VirtualProduction/EpicStageApp/Source/EpicStageApp/Private/StageAppRouteHandler.h
@@ -127,8 +127,8 @@ private:
 	/** Handles ending an actor drag operation. */
 	void HandleWebSocketNDisplayPreviewActorDragEnd(const FRemoteControlWebSocketMessage& WebSocketMessage);
 
-	/** Handles creating a lightcard. */
-	void HandleWebSocketNDisplayPreviewLightCardCreate(const FRemoteControlWebSocketMessage& WebSocketMessage);
+	/** Handles creating an actor. */
+	void HandleWebSocketNDisplayPreviewActorCreate(const FRemoteControlWebSocketMessage& WebSocketMessage);
 
 	/** Called when a client disconnects from the WebSocket server. */
 	void HandleClientDisconnected(FGuid ClientId);
