[EpicStageApp] Fix LCs being created from the app without an associated transaction

diff --git a/Engine/Plugins/VirtualProduction/EpicStageApp/Source/EpicStageApp/Private/StageAppRouteHandler.cpp b/Engine/Plugins/VirtualProduction/EpicStageApp/Source/EpicStageApp/Private/StageAppRouteHandler.cpp
--- a/Engine/Plugins/VirtualProduction/EpicStageApp/Source/EpicStageApp/Private/StageAppRouteHandler.cpp
+++ b/Engine/Plugins/VirtualProduction/EpicStageApp/Source/EpicStageApp/Private/StageAppRouteHandler.cpp
@@ -786,6 +786,10 @@ void FStageAppRouteHandler::HandleWebSocketNDisplayPreviewActorCreate(const FRem
 		SpawnArgs.ProjectionMode = PerRendererData->GetLightCardHelper().GetProjectionMode();
 	}
 
+#if WITH_EDITOR
+	FScopedTransaction Transaction(LOCTEXT("SpawnActorTransaction", "Spawn Actor from Stage App"));
+#endif
+
 	AActor* NewActor = FDisplayClusterLightCardEditorHelper::SpawnStageActor(SpawnArgs);
 	if (!NewActor)
 	{
