[EpicStageApp] Fix duplicate button failing for some actors (C++ work)

diff --git a/Engine/Plugins/VirtualProduction/EpicStageApp/Source/EpicStageApp/Private/StageAppRequest.h b/Engine/Plugins/VirtualProduction/EpicStageApp/Source/EpicStageApp/Private/StageAppRequest.h
--- a/Engine/Plugins/VirtualProduction/EpicStageApp/Source/EpicStageApp/Private/StageAppRequest.h
+++ b/Engine/Plugins/VirtualProduction/EpicStageApp/Source/EpicStageApp/Private/StageAppRequest.h
@@ -392,3 +392,34 @@ struct FRCWebSocketNDisplayPreviewActorCreateBody : public FRCRequest
 	int32 RequestId = -1;
 };
 
+/**
+ * Holds a request made via websocket to duplicate one or more actors.
+ */
+USTRUCT()
+struct FRCWebSocketNDisplayActorDuplicateBody : public FRCRequest
+{
+	GENERATED_BODY()
+
+	FRCWebSocketNDisplayActorDuplicateBody()
+	{
+		AddStructParameter(ParametersFieldLabel());
+	}
+
+	/**
+	 * Get the label for the property value struct.
+	 */
+	static FString ParametersFieldLabel() { return TEXT("Parameters"); }
+
+	/**
+	 * The list of paths of actors to duplicate.
+	 */
+	UPROPERTY()
+	TArray<FString> Actors;
+
+	/**
+	 * An optional number that will be passed back in the RequestedActorsCreated response to tell apart
+	 * the results of multiple requests.
+	 */
+	UPROPERTY()
+	int32 RequestId = -1;
+};
\ No newline at end of file
diff --git a/Engine/Plugins/VirtualProduction/EpicStageApp/Source/EpicStageApp/Private/StageAppRouteHandler.cpp b/Engine/Plugins/VirtualProduction/EpicStageApp/Source/EpicStageApp/Private/StageAppRouteHandler.cpp
--- a/Engine/Plugins/VirtualProduction/EpicStageApp/Source/EpicStageApp/Private/StageAppRouteHandler.cpp
+++ b/Engine/Plugins/VirtualProduction/EpicStageApp/Source/EpicStageApp/Private/StageAppRouteHandler.cpp
@@ -98,6 +98,53 @@ namespace StageAppRouteHandlerUtils
 			return EDisplayClusterMeshProjectionOutput::Color;
 		}
 	}
+
+	/// Given a light card, find the ADisplayClusterRootActor to which it belongs (or null if none).
+	ADisplayClusterRootActor* GetRootActorForLightCard(AActor* InActor)
+	{
+		UWorld* World = InActor->GetWorld();
+		if (!World)
+		{
+			return nullptr;
+		}
+
+		for (TActorIterator<ADisplayClusterRootActor> ActorIterator(World); ActorIterator; ++ActorIterator)
+		{
+			ADisplayClusterRootActor* RootActor = *ActorIterator;
+			if (!RootActor)
+			{
+				continue;
+			}
+
+			const UDisplayClusterConfigurationData* ConfigData = RootActor->GetConfigData();
+			if (!ConfigData)
+			{
+				continue;
+			}
+
+			const FDisplayClusterConfigurationICVFX_VisibilityList& RootActorLightCards = ConfigData->StageSettings.Lightcard.ShowOnlyList;
+
+			// Check if it's in the root actor's layers
+			for (const FActorLayer& ActorLayer : RootActorLightCards.ActorLayers)
+			{
+				if (InActor->Layers.Contains(ActorLayer.Name))
+				{
+					return RootActor;
+				}
+			}
+
+			// Check if it's directly on the root actor
+			if (RootActorLightCards.Actors.ContainsByPredicate([InActor](const TSoftObjectPtr<AActor>& Actor)
+				{
+					return InActor == Actor.Get();
+				}))
+			{
+				return RootActor;
+			}
+		}
+
+		return nullptr;
+	}
 }
 
 
@@ -281,6 +328,12 @@ void FStageAppRouteHandler::RegisterRoutes(IWebRemoteControlModule& WebRemoteCon
 		TEXT("ndisplay.preview.actor.create"),
 		FWebSocketMessageDelegate::CreateRaw(this, &FStageAppRouteHandler::HandleWebSocketNDisplayPreviewActorCreate)
 	));
+
+	RegisterRoute(MakeUnique<FRemoteControlWebsocketRoute>(
+		TEXT("Duplicate existing actors"),
+		TEXT("stageapp.actors.duplicate"),
+		FWebSocketMessageDelegate::CreateRaw(this, &FStageAppRouteHandler::HandleWebSocketStageAppActorsDuplicate)
+	));
 }
 
 void FStageAppRouteHandler::RegisterEngineEvents()
@@ -826,6 +879,93 @@ void FStageAppRouteHandler::HandleWebSocketNDisplayPreviewActorCreate(const FRem
 	RemoteControlModule->SendWebsocketMessage(WebSocketMessage.ClientId, Payload);
 }
 
+void FStageAppRouteHandler::HandleWebSocketStageAppActorsDuplicate(const FRemoteControlWebSocketMessage& WebSocketMessage)
+{
+	FRCWebSocketNDisplayActorDuplicateBody Body;
+	if (!WebRemoteControlUtils::DeserializeMessage(WebSocketMessage.RequestPayload, Body))
+	{
+		return;
+	}
+
+	FRCRequestedActorsCreated Event;
+
+#if WITH_EDITOR
+	TMap<UWorld*, FCachedActorLabels> ActorLabelsByWorld;
+	FScopedTransaction Transaction(LOCTEXT("DuplicateActors", "Duplicate Actors from Stage App"));
+#endif
+
+	for (const FString& ActorPath : Body.Actors)
+	{
+		AActor* Actor = FindObject<AActor>(nullptr, *ActorPath);
+		if (!Actor)
+		{
+			continue;
+		}
+
+		UWorld* World = Actor->GetWorld();
+		if (!World)
+		{
+			continue;
+		}
+
+		ULevel* Level = World->GetCurrentLevel();
+		if (!Level)
+		{
+			continue;
+		}
+
+		const FName UniqueName = MakeUniqueObjectName(Level, Actor->GetClass());
+		AActor* NewActor = CastChecked<AActor>(StaticDuplicateObject(Actor, Level, UniqueName));
+
+#if WITH_EDITOR
+		Level->AddLoadedActor(NewActor);
+
+		FCachedActorLabels* ActorLabels = ActorLabelsByWorld.Find(World);
+		if (!ActorLabels)
+		{
+			ActorLabels = &ActorLabelsByWorld.Emplace(World, World);
+		}
+
+		FActorLabelUtilities::SetActorLabelUnique(NewActor, NewActor->GetActorLabel(), ActorLabels);
+		ActorLabels->Add(NewActor->GetActorLabel());
+#endif
+
+		// Special handling for lightcards
+		if (ADisplayClusterLightCardActor* NewLightCard = Cast<ADisplayClusterLightCardActor>(NewActor))
+		{
+			// Add it to the root actor
+			if (ADisplayClusterRootActor* RootActor = StageAppRouteHandlerUtils::GetRootActorForLightCard(Actor))
+			{
+				NewLightCard->AddToLightCardLayer(RootActor);
+
+#if WITH_EDITOR
+				// Required so operator panel updates
+				RootActor->PostEditChange();
+#endif
+			}
+		}
+
+#if WITH_EDITOR
+		if (GIsEditor)
+		{
+			GEditor->BroadcastLevelActorAdded(NewActor);
+		}
+#endif
+
+		Event.ActorPaths.Add(NewActor->GetPathName());
+	}
+
+	// Send new actors back to the client
+	if (!Event.ActorPaths.IsEmpty())
+	{
+		Event.RequestId = Body.RequestId;
+
+		TArray<uint8> Payload;
+		WebRemoteControlUtils::SerializeMessage(Event, Payload);
+		RemoteControlModule->SendWebsocketMessage(WebSocketMessage.ClientId, Payload);
+	}
+}
+
 void FStageAppRouteHandler::HandleClientDisconnected(FGuid ClientId)
 {
 	// Destroy the client's renderers
diff --git a/Engine/Plugins/VirtualProduction/EpicStageApp/Source/EpicStageApp/Private/StageAppRouteHandler.h b/Engine/Plugins/VirtualProduction/EpicStageApp/Source/EpicStageApp/Private/StageAppRouteHandler.h
--- a/Engine/Plugins/VirtualProduction/EpicStageApp/Source/EpicStageApp/Private/StageAppRouteHandler.h
+++ b/Engine/Plugins/VirtualProduction/EpicStageApp/Source/EpicStageApp/Private/StageAppRouteHandler.h
@@ -130,6 +130,9 @@ private:
 	/** Handles creating an actor. */
 	void HandleWebSocketNDisplayPreviewActorCreate(const FRemoteControlWebSocketMessage& WebSocketMessage);
 
+	/** Handles duplicating actors. */
+	void HandleWebSocketStageAppActorsDuplicate(const FRemoteControlWebSocketMessage& WebSocketMessage);
+
 	/** Called when a client disconnects from the WebSocket server. */
 	void HandleClientDisconnected(FGuid ClientId);
 
