[EpicStageApp] Fix cases preventing remote creation of flag light cards

diff --git a/Engine/Plugins/VirtualProduction/EpicStageApp/Source/EpicStageApp/Private/StageAppRequest.h b/Engine/Plugins/VirtualProduction/EpicStageApp/Source/EpicStageApp/Private/StageAppRequest.h
--- a/Engine/Plugins/VirtualProduction/EpicStageApp/Source/EpicStageApp/Private/StageAppRequest.h
+++ b/Engine/Plugins/VirtualProduction/EpicStageApp/Source/EpicStageApp/Private/StageAppRequest.h
@@ -365,7 +365,9 @@ struct FRCWebSocketNDisplayPreviewActorCreateBody : public FRCRequest
 	FString ActorClass = "";
 
 	/**
-	 * The path of the template to use for the lightcard. If empty, a lightcard will be created with default settings.
+	 * The path of the template to use for the lightcard.
+	 * If empty, a lightcard will be created using the default template.
+	 * If "None", a lightcard will be created with default settings regardless of whether there's a default template.
 	 */
 	UPROPERTY()
 	FString TemplatePath = "";
@@ -384,6 +386,19 @@ struct FRCWebSocketNDisplayPreviewActorCreateBody : public FRCRequest
 	UPROPERTY()
 	FVector2D Position = FVector2D::ZeroVector;
 
+	/**
+	 * If true, override the default/template color for the actor if it's a lightcard.
+	 */
+	UPROPERTY()
+	bool OverrideColor = false;
+
+	/**
+	 * If OverrideColor is true and PreviewRendererId points to a valid preview renderer, use this color when creating
+	 * a lightcard.
+	 */
+	UPROPERTY()
+	FLinearColor Color = FLinearColor::White;
+
 	/**
 	 * An optional number that will be passed back in the RequestedActorsCreated response to tell apart
 	 * the results of multiple requests.
diff --git a/Engine/Plugins/VirtualProduction/EpicStageApp/Source/EpicStageApp/Private/StageAppRouteHandler.cpp b/Engine/Plugins/VirtualProduction/EpicStageApp/Source/EpicStageApp/Private/StageAppRouteHandler.cpp
--- a/Engine/Plugins/VirtualProduction/EpicStageApp/Source/EpicStageApp/Private/StageAppRouteHandler.cpp
+++ b/Engine/Plugins/VirtualProduction/EpicStageApp/Source/EpicStageApp/Private/StageAppRouteHandler.cpp
@@ -819,7 +819,10 @@ void FStageAppRouteHandler::HandleWebSocketNDisplayPreviewActorCreate(const FRem
 
 	if (!Body.TemplatePath.IsEmpty())
 	{
-		Template = Cast<UDisplayClusterStageActorTemplate>(StaticLoadObject(UDisplayClusterStageActorTemplate::StaticClass(), nullptr, *Body.TemplatePath));
+		if (Body.TemplatePath != TEXT("None"))
+		{
+			Template = Cast<UDisplayClusterStageActorTemplate>(StaticLoadObject(UDisplayClusterStageActorTemplate::StaticClass(), nullptr, *Body.TemplatePath));
+		}
 	}
 #if WITH_EDITOR
 	else if (ActorClass == ADisplayClusterLightCardActor::StaticClass())
@@ -849,6 +852,14 @@ void FStageAppRouteHandler::HandleWebSocketNDisplayPreviewActorCreate(const FRem
 		return;
 	}
 
+	if (ADisplayClusterLightCardActor* LightCard = Cast<ADisplayClusterLightCardActor>(NewActor))
+	{
+		if (Body.OverrideColor)
+		{
+			LightCard->Color = Body.Color;
+		}
+	}
+
 	// Override the location
 	if (Body.OverridePosition && NewActor->Implements<UDisplayClusterStageActor>())
 	{
