[EpicStageApp] Semantic version for stage app API

diff --git a/Engine/Plugins/VirtualProduction/EpicStageApp/Source/EpicStageApp/Private/StageAppLibrary.cpp b/Engine/Plugins/VirtualProduction/EpicStageApp/Source/EpicStageApp/Private/StageAppLibrary.cpp
new file mode 100644
--- /dev/null
+++ b/Engine/Plugins/VirtualProduction/EpicStageApp/Source/EpicStageApp/Private/StageAppLibrary.cpp
@@ -0,0 +1,12 @@
+// Copyright Epic Games, Inc. All Rights Reserved.
+
+#include "StageAppLibrary.h"
+#include "StageAppVersion.h"
+
+FString UStageAppFunctionLibrary::GetAPIVersion()
+{
+	return FString::Printf(TEXT("%u.%u.%u"),
+		FEpicStageAppAPIVersion::Major,
+		FEpicStageAppAPIVersion::Minor,
+		FEpicStageAppAPIVersion::Patch);
+}
diff --git a/Engine/Plugins/VirtualProduction/EpicStageApp/Source/EpicStageApp/Public/StageAppLibrary.h b/Engine/Plugins/VirtualProduction/EpicStageApp/Source/EpicStageApp/Public/StageAppLibrary.h
new file mode 100644
--- /dev/null
+++ b/Engine/Plugins/VirtualProduction/EpicStageApp/Source/EpicStageApp/Public/StageAppLibrary.h
@@ -0,0 +1,22 @@
+// Copyright Epic Games, Inc. All Rights Reserved.
+
+#pragma once
+
+#include "CoreMinimal.h"
+#include "Kismet/BlueprintFunctionLibrary.h"
+#include "StageAppLibrary.generated.h"
+
+/**
+ * Generally useful Blueprint/remote functions for Epic Stage App integration.
+ */
+UCLASS()
+class EPICSTAGEAPP_API UStageAppFunctionLibrary : public UBlueprintFunctionLibrary
+{
+	GENERATED_BODY()
+
+public:
+
+	/** The current semantic version for the stage app API as a formatted string. */
+	UFUNCTION(BlueprintPure, Category = "Development", meta = (BlueprintThreadSafe))
+	static FString GetAPIVersion();
+};
diff --git a/Engine/Plugins/VirtualProduction/EpicStageApp/Source/EpicStageApp/Public/StageAppVersion.h b/Engine/Plugins/VirtualProduction/EpicStageApp/Source/EpicStageApp/Public/StageAppVersion.h
new file mode 100644
--- /dev/null
+++ b/Engine/Plugins/VirtualProduction/EpicStageApp/Source/EpicStageApp/Public/StageAppVersion.h
@@ -0,0 +1,18 @@
+// Copyright Epic Games, Inc. All Rights Reserved.
+
+#pragma once
+
+#include "CoreMinimal.h"
+
+/** Holds info about a version of the stage app API. */
+struct EPICSTAGEAPP_API FEpicStageAppAPIVersion
+{
+	/** The major component of the current version. */
+	static const uint16 Major = 1;
+
+	/** The minor component of the current version. */
+	static const uint16 Minor = 0;
+	
+	/** The patch component of the current version. */
+	static const uint16 Patch = 0;
+};
\ No newline at end of file
