[EpicStageApp] Expose Remote Control webserver IP via Stage App module

diff --git a/Engine/Plugins/VirtualProduction/EpicStageApp/Source/EpicStageApp/Private/StageAppLibrary.cpp b/Engine/Plugins/VirtualProduction/EpicStageApp/Source/EpicStageApp/Private/StageAppLibrary.cpp
--- a/Engine/Plugins/VirtualProduction/EpicStageApp/Source/EpicStageApp/Private/StageAppLibrary.cpp
+++ b/Engine/Plugins/VirtualProduction/EpicStageApp/Source/EpicStageApp/Private/StageAppLibrary.cpp
@@ -2,6 +2,7 @@
 
 #include "StageAppLibrary.h"
 #include "StageAppVersion.h"
+#include "RemoteControlSettings.h"
 
 FString UStageAppFunctionLibrary::GetAPIVersion()
 {
@@ -10,3 +11,8 @@ FString UStageAppFunctionLibrary::GetAPIVersion()
 		FEpicStageAppAPIVersion::Minor,
 		FEpicStageAppAPIVersion::Patch);
 }
+
+int32 UStageAppFunctionLibrary::GetRemoteControlWebInterfacePort()
+{
+	return static_cast<int32>(GetDefault<URemoteControlSettings>()->RemoteControlWebInterfacePort);
+}
\ No newline at end of file
diff --git a/Engine/Plugins/VirtualProduction/EpicStageApp/Source/EpicStageApp/Public/StageAppLibrary.h b/Engine/Plugins/VirtualProduction/EpicStageApp/Source/EpicStageApp/Public/StageAppLibrary.h
--- a/Engine/Plugins/VirtualProduction/EpicStageApp/Source/EpicStageApp/Public/StageAppLibrary.h
+++ b/Engine/Plugins/VirtualProduction/EpicStageApp/Source/EpicStageApp/Public/StageAppLibrary.h
@@ -16,7 +16,11 @@ class EPICSTAGEAPP_API UStageAppFunctionLibrary : public UBlueprintFunctionLibra
 
 public:
 
-	/** The current semantic version for the stage app API as a formatted string. */
+	/** Get the current semantic version for the stage app API as a formatted string. */
 	UFUNCTION(BlueprintPure, Category = "Development", meta = (BlueprintThreadSafe))
 	static FString GetAPIVersion();
+
+	/** Get the port number used to access the remote control web interface for this engine instance. */
+	UFUNCTION(BlueprintPure, Category = "Development", meta = (BlueprintThreadSafe))
+	static int32 GetRemoteControlWebInterfacePort();
 };
diff --git a/Engine/Plugins/VirtualProduction/EpicStageApp/Source/EpicStageApp/Public/StageAppVersion.h b/Engine/Plugins/VirtualProduction/EpicStageApp/Source/EpicStageApp/Public/StageAppVersion.h
--- a/Engine/Plugins/VirtualProduction/EpicStageApp/Source/EpicStageApp/Public/StageAppVersion.h
+++ b/Engine/Plugins/VirtualProduction/EpicStageApp/Source/EpicStageApp/Public/StageAppVersion.h
@@ -11,7 +11,7 @@ struct EPICSTAGEAPP_API FEpicStageAppAPIVersion
 	static const uint16 Major = 1;
 
 	/** The minor component of the current version. */
-	static const uint16 Minor = 1;
+	static const uint16 Minor = 2;
 	
 	/** The patch component of the current version. */
 	static const uint16 Patch = 0;
