[EpicStageApp] Bump to v1.3.0 to disambiguate routes added in 24506076, which aren't included in 5.2 (v1.2.0)

diff --git a/Engine/Plugins/VirtualProduction/EpicStageApp/Source/EpicStageApp/Public/StageAppVersion.h b/Engine/Plugins/VirtualProduction/EpicStageApp/Source/EpicStageApp/Public/StageAppVersion.h
--- a/Engine/Plugins/VirtualProduction/EpicStageApp/Source/EpicStageApp/Public/StageAppVersion.h
+++ b/Engine/Plugins/VirtualProduction/EpicStageApp/Source/EpicStageApp/Public/StageAppVersion.h
@@ -11,7 +11,7 @@ struct EPICSTAGEAPP_API FEpicStageAppAPIVersion
 	static const uint16 Major = 1;
 
 	/** The minor component of the current version. */
-	static const uint16 Minor = 2;
+	static const uint16 Minor = 3;
 	
 	/** The patch component of the current version. */
 	static const uint16 Patch = 0;
