[EpicStageApp] Bump stage app API to reflect new changes (v1.3.0 is now lightcard subtypes and v1.4.0 is array HTTP routes)

diff --git a/Engine/Plugins/VirtualProduction/EpicStageApp/Source/EpicStageApp/Public/StageAppVersion.h b/Engine/Plugins/VirtualProduction/EpicStageApp/Source/EpicStageApp/Public/StageAppVersion.h
--- a/Engine/Plugins/VirtualProduction/EpicStageApp/Source/EpicStageApp/Public/StageAppVersion.h
+++ b/Engine/Plugins/VirtualProduction/EpicStageApp/Source/EpicStageApp/Public/StageAppVersion.h
@@ -11,7 +11,7 @@ struct EPICSTAGEAPP_API FEpicStageAppAPIVersion
 	static const uint16 Major = 1;
 
 	/** The minor component of the current version. */
-	static const uint16 Minor = 3;
+	static const uint16 Minor = 4;
 	
 	/** The patch component of the current version. */
 	static const uint16 Patch = 0;
