[EpicStageApp] Update actor creation API to allow for bIsLightCardFlag property

diff --git a/Engine/Plugins/VirtualProduction/EpicStageApp/Source/EpicStageApp/Private/StageAppRequest.h b/Engine/Plugins/VirtualProduction/EpicStageApp/Source/EpicStageApp/Private/StageAppRequest.h
--- a/Engine/Plugins/VirtualProduction/EpicStageApp/Source/EpicStageApp/Private/StageAppRequest.h
+++ b/Engine/Plugins/VirtualProduction/EpicStageApp/Source/EpicStageApp/Private/StageAppRequest.h
@@ -399,6 +399,12 @@ struct FRCWebSocketNDisplayPreviewActorCreateBody : public FRCRequest
 	UPROPERTY()
 	FLinearColor Color = FLinearColor::White;
 
+	/**
+	 * Whether this actor should be created as a flag if it's a lightcard.
+	 */
+	UPROPERTY()
+	bool bIsLightCardFlag = false;
+
 	/**
 	 * An optional number that will be passed back in the RequestedActorsCreated response to tell apart
 	 * the results of multiple requests.
diff --git a/Engine/Plugins/VirtualProduction/EpicStageApp/Source/EpicStageApp/Private/StageAppRouteHandler.cpp b/Engine/Plugins/VirtualProduction/EpicStageApp/Source/EpicStageApp/Private/StageAppRouteHandler.cpp
--- a/Engine/Plugins/VirtualProduction/EpicStageApp/Source/EpicStageApp/Private/StageAppRouteHandler.cpp
+++ b/Engine/Plugins/VirtualProduction/EpicStageApp/Source/EpicStageApp/Private/StageAppRouteHandler.cpp
@@ -858,6 +858,8 @@ void FStageAppRouteHandler::HandleWebSocketNDisplayPreviewActorCreate(const FRem
 		{
 			LightCard->Color = Body.Color;
 		}
+
+		LightCard->SetIsLightCardFlag(Body.bIsLightCardFlag);
 	}
 
 	// Override the location
