[EpicStageApp] Add API to create objects via WebSocket and keep them in memory as long as the connection remains open

diff --git a/Engine/Plugins/VirtualProduction/EpicStageApp/Source/EpicStageApp/Private/StageAppRequest.h b/Engine/Plugins/VirtualProduction/EpicStageApp/Source/EpicStageApp/Private/StageAppRequest.h
--- a/Engine/Plugins/VirtualProduction/EpicStageApp/Source/EpicStageApp/Private/StageAppRequest.h
+++ b/Engine/Plugins/VirtualProduction/EpicStageApp/Source/EpicStageApp/Private/StageAppRequest.h
@@ -443,4 +443,70 @@ struct FRCWebSocketNDisplayActorDuplicateBody : public FRCRequest
 	 */
 	UPROPERTY()
 	int32 RequestId = -1;
-};
\ No newline at end of file
+};
+
+/**
+ * Holds a request made via websocket to create an object whose lifetime will be bound to the lifetime
+ * of the requesting client's connection.
+ */
+USTRUCT()
+struct FRCWebSocketStageAppClientObjectCreateBody : public FRCRequest
+{
+	GENERATED_BODY()
+
+	FRCWebSocketStageAppClientObjectCreateBody()
+	{
+		AddStructParameter(ParametersFieldLabel());
+	}
+
+	/**
+	 * Get the label for the property value struct.
+	 */
+	static FString ParametersFieldLabel() { return TEXT("Parameters"); }
+
+	/**
+	 * The class of the object to create.
+	 */
+	UPROPERTY()
+	TSubclassOf<UObject> ObjectClass;
+
+	/**
+	 * An optional number that will be passed back in the ClientObjectCreated response to tell apart
+	 * the results of multiple requests.
+	 */
+	UPROPERTY()
+	int32 RequestId = -1;
+};
+
+/**
+ * Holds a request made via websocket to destroy an object whose lifetime is bound to the lifetime
+ * of the requesting client's connection.
+ */
+USTRUCT()
+struct FRCWebSocketStageAppClientObjectDestroyBody : public FRCRequest
+{
+	GENERATED_BODY()
+
+	FRCWebSocketStageAppClientObjectDestroyBody()
+	{
+		AddStructParameter(ParametersFieldLabel());
+	}
+
+	/**
+	 * Get the label for the property value struct.
+	 */
+	static FString ParametersFieldLabel() { return TEXT("Parameters"); }
+
+	/**
+	 * The path of the object. Must be an object created using the stageapp.clientobject.create route.
+	 */
+	UPROPERTY()
+	FString ObjectPath;
+
+	/**
+	 * An optional number that will be passed back in the RequestedClientObjectCreated response to tell
+	 * apart the results of multiple requests.
+	 */
+	UPROPERTY()
+	int32 RequestId = -1;
+};
diff --git a/Engine/Plugins/VirtualProduction/EpicStageApp/Source/EpicStageApp/Private/StageAppResponse.h b/Engine/Plugins/VirtualProduction/EpicStageApp/Source/EpicStageApp/Private/StageAppResponse.h
--- a/Engine/Plugins/VirtualProduction/EpicStageApp/Source/EpicStageApp/Private/StageAppResponse.h
+++ b/Engine/Plugins/VirtualProduction/EpicStageApp/Source/EpicStageApp/Private/StageAppResponse.h
@@ -158,3 +158,35 @@ struct FRCRequestedActorsCreated
 	UPROPERTY()
 	int32 RequestId = -1;
 };
+
+/**
+ * Event sent to a client when a client-bound object was created at the request of the client.
+ */
+USTRUCT()
+struct FRCRequestedClientObjectCreated
+{
+	GENERATED_BODY()
+
+	FRCRequestedClientObjectCreated()
+	: Type(TEXT("RequestedClientObjectCreated"))
+	{
+	}
+
+	/**
+	 * Type of the event.
+	 */
+	UPROPERTY()
+	FString Type;
+
+	/**
+	 * The list of newly-created object's path.
+	 */
+	UPROPERTY()
+	FString ObjectPath;
+
+	/**
+	 * The ID passed when the actors were created, or -1 if no ID was provided.
+	 */
+	UPROPERTY()
+	int32 RequestId = -1;
+};
diff --git a/Engine/Plugins/VirtualProduction/EpicStageApp/Source/EpicStageApp/Private/StageAppRouteHandler.cpp b/Engine/Plugins/VirtualProduction/EpicStageApp/Source/EpicStageApp/Private/StageAppRouteHandler.cpp
--- a/Engine/Plugins/VirtualProduction/EpicStageApp/Source/EpicStageApp/Private/StageAppRouteHandler.cpp
+++ b/Engine/Plugins/VirtualProduction/EpicStageApp/Source/EpicStageApp/Private/StageAppRouteHandler.cpp
@@ -334,6 +334,18 @@ void FStageAppRouteHandler::RegisterRoutes(IWebRemoteControlModule& WebRemoteCon
 		TEXT("stageapp.actors.duplicate"),
 		FWebSocketMessageDelegate::CreateRaw(this, &FStageAppRouteHandler::HandleWebSocketStageAppActorsDuplicate)
 	));
+
+	RegisterRoute(MakeUnique<FRemoteControlWebsocketRoute>(
+		TEXT("Create an object that will be kept alive until the client disconnects"),
+		TEXT("stageapp.clientobject.create"),
+		FWebSocketMessageDelegate::CreateRaw(this, &FStageAppRouteHandler::HandleWebSocketStageAppClientObjectCreate)
+	));
+
+	RegisterRoute(MakeUnique<FRemoteControlWebsocketRoute>(
+		TEXT("Destroy an object that was created using the stageapp.clientobject.create route"),
+		TEXT("stageapp.clientobject.destroy"),
+		FWebSocketMessageDelegate::CreateRaw(this, &FStageAppRouteHandler::HandleWebSocketStageAppClientObjectDestroy)
+	));
 }
 
 void FStageAppRouteHandler::RegisterEngineEvents()
@@ -979,6 +991,70 @@ void FStageAppRouteHandler::HandleWebSocketStageAppActorsDuplicate(const FRemote
 	}
 }
 
+void FStageAppRouteHandler::HandleWebSocketStageAppClientObjectCreate(const FRemoteControlWebSocketMessage& WebSocketMessage)
+{
+	FRCWebSocketStageAppClientObjectCreateBody Body;
+	if (!WebRemoteControlUtils::DeserializeMessage(WebSocketMessage.RequestPayload, Body))
+	{
+		return;
+	}
+
+	UPackage* Package = CreatePackage(TEXT("/Engine/Transient/StageAppClientObjects"));
+	if (!Package)
+	{
+		return;
+	}
+
+	UObject* CreatedObject = NewObject<UObject>(Package, Body.ObjectClass, NAME_None, RF_Transient | RF_Public);
+	if (!CreatedObject)
+	{
+		return;
+	}
+
+	CreatedObject->AddToRoot();
+	BoundObjectsByClientId.FindOrAdd(WebSocketMessage.ClientId).Add(CreatedObject);
+
+	// Sent response with the new object's path
+	FRCRequestedClientObjectCreated Event;
+
+	Event.RequestId = Body.RequestId;
+	Event.ObjectPath = CreatedObject->GetPathName();
+
+	TArray<uint8> Payload;
+	WebRemoteControlUtils::SerializeMessage(Event, Payload);
+	RemoteControlModule->SendWebsocketMessage(WebSocketMessage.ClientId, Payload);
+}
+
+void FStageAppRouteHandler::HandleWebSocketStageAppClientObjectDestroy(const FRemoteControlWebSocketMessage& WebSocketMessage)
+{
+	FRCWebSocketStageAppClientObjectDestroyBody Body;
+	if (!WebRemoteControlUtils::DeserializeMessage(WebSocketMessage.RequestPayload, Body))
+	{
+		return;
+	}
+
+	if (TArray<UObject*>* ClientBoundObjects = BoundObjectsByClientId.Find(WebSocketMessage.ClientId))
+	{
+		for (auto BoundObjectIt = ClientBoundObjects->CreateIterator(); BoundObjectIt; ++BoundObjectIt)
+		{
+			UObject* BoundObject = *BoundObjectIt;
+			if (!BoundObject)
+			{
+				BoundObjectIt.RemoveCurrentSwap();
+				continue;
+			}
+
+			if (BoundObject->GetPathName() == Body.ObjectPath)
+			{
+				BoundObjectIt.RemoveCurrentSwap();
+				BoundObject->RemoveFromRoot();
+				BoundObject->ConditionalBeginDestroy();
+				return;
+			}
+		}
+	}
+}
+
 void FStageAppRouteHandler::HandleClientDisconnected(FGuid ClientId)
 {
 	// Destroy the client's renderers
@@ -993,6 +1069,18 @@ void FStageAppRouteHandler::HandleClientDisconnected(FGuid ClientId)
 
 		PerRendererDataMapsByClientId.Remove(ClientId);
 	}
+
+	// Destroy objects bound to the client's connection lifetime
+	if (TArray<UObject*>* ClientBoundObjects = BoundObjectsByClientId.Find(ClientId))
+	{
+		for (UObject* BoundObject : *ClientBoundObjects)
+		{
+			BoundObject->RemoveFromRoot();
+			BoundObject->ConditionalBeginDestroy();
+		}
+
+		BoundObjectsByClientId.Remove(ClientId);
+	}
 }
 
 void FStageAppRouteHandler::ChangePreviewRendererSettings(const FGuid& ClientId, int32 RendererId, const FRCWebSocketNDisplayPreviewRendererSettings& NewSettings)
diff --git a/Engine/Plugins/VirtualProduction/EpicStageApp/Source/EpicStageApp/Private/StageAppRouteHandler.h b/Engine/Plugins/VirtualProduction/EpicStageApp/Source/EpicStageApp/Private/StageAppRouteHandler.h
--- a/Engine/Plugins/VirtualProduction/EpicStageApp/Source/EpicStageApp/Private/StageAppRouteHandler.h
+++ b/Engine/Plugins/VirtualProduction/EpicStageApp/Source/EpicStageApp/Private/StageAppRouteHandler.h
@@ -133,6 +133,12 @@ private:
 	/** Handles duplicating actors. */
 	void HandleWebSocketStageAppActorsDuplicate(const FRemoteControlWebSocketMessage& WebSocketMessage);
 
+	/** Handles creating an object whose lifetime is tied to the lifetime of the requesting client's connection. */
+	void HandleWebSocketStageAppClientObjectCreate(const FRemoteControlWebSocketMessage& WebSocketMessage);
+
+	/** Handles destroying an object whose lifetime is tied to the lifetime of the requesting client's connection. */
+	void HandleWebSocketStageAppClientObjectDestroy(const FRemoteControlWebSocketMessage& WebSocketMessage);
+
 	/** Called when a client disconnects from the WebSocket server. */
 	void HandleClientDisconnected(FGuid ClientId);
 
@@ -168,6 +174,9 @@ private:
 	/** Map from client GUID to preview renderer data maps for that client. */
 	TMap<FGuid, TClientIdToPerRendererDataMap> PerRendererDataMapsByClientId;
 
+	/** Map from client GUID to objects bound to the lifetime of that client's connection. */
+	TMap<FGuid, TArray<UObject*>> BoundObjectsByClientId;
+
 	/** Image wrapper module used to compress preview images. */
 	class IImageWrapperModule* ImageWrapperModule = nullptr;
 
diff --git a/Engine/Plugins/VirtualProduction/EpicStageApp/Source/EpicStageApp/Public/StageAppVersion.h b/Engine/Plugins/VirtualProduction/EpicStageApp/Source/EpicStageApp/Public/StageAppVersion.h
--- a/Engine/Plugins/VirtualProduction/EpicStageApp/Source/EpicStageApp/Public/StageAppVersion.h
+++ b/Engine/Plugins/VirtualProduction/EpicStageApp/Source/EpicStageApp/Public/StageAppVersion.h
@@ -11,7 +11,7 @@ struct EPICSTAGEAPP_API FEpicStageAppAPIVersion
 	static const uint16 Major = 1;
 
 	/** The minor component of the current version. */
-	static const uint16 Minor = 6;
+	static const uint16 Minor = 7;
 	
 	/** The patch component of the current version. */
 	static const uint16 Patch = 0;
