[EpicStageApp] Fix no-PCH compilation errors

diff --git a/Engine/Plugins/VirtualProduction/EpicStageApp/Source/EpicStageApp/Private/StageAppRouteHandler.cpp b/Engine/Plugins/VirtualProduction/EpicStageApp/Source/EpicStageApp/Private/StageAppRouteHandler.cpp
--- a/Engine/Plugins/VirtualProduction/EpicStageApp/Source/EpicStageApp/Private/StageAppRouteHandler.cpp
+++ b/Engine/Plugins/VirtualProduction/EpicStageApp/Source/EpicStageApp/Private/StageAppRouteHandler.cpp
@@ -18,6 +18,7 @@
 #include "RemoteControlWebsocketRoute.h"
 #include "StageAppResponse.h"
 #include "TextureResource.h"
+#include "UObject/Package.h"
 #include "WebRemoteControlUtils.h"
 #include "StageActor/DisplayClusterStageActorTemplate.h"
 #include "StageActor/DisplayClusterWeakStageActorPtr.h"
