[UnrealStage] Fix local DCRA rotation being ignored when creating a stage actor remotely

diff --git a/Engine/Plugins/Runtime/nDisplay/Source/DisplayClusterScenePreview/Private/DisplayClusterLightCardEditorHelper.cpp b/Engine/Plugins/Runtime/nDisplay/Source/DisplayClusterScenePreview/Private/DisplayClusterLightCardEditorHelper.cpp
--- a/Engine/Plugins/Runtime/nDisplay/Source/DisplayClusterScenePreview/Private/DisplayClusterLightCardEditorHelper.cpp
+++ b/Engine/Plugins/Runtime/nDisplay/Source/DisplayClusterScenePreview/Private/DisplayClusterLightCardEditorHelper.cpp
@@ -323,7 +323,8 @@ void FDisplayClusterLightCardEditorHelper::MoveActorsToPixel(
 		AverageCoords.Conform();
 
 		// Compute desired coordinates (radius doesn't matter here since we will use the flush constraint on the light cards after moving them)
-		const FSphericalCoordinates DesiredCoords(Direction * 100.0f);
+		const FVector LocalDirection = CachedRootActor->GetActorRotation().RotateVector(Direction);
+		const FSphericalCoordinates DesiredCoords(LocalDirection * 100.0f);
 		const FSphericalCoordinates DeltaCoords = DesiredCoords - AverageCoords;
 
 		// Update each light card with the delta coordinates; the flush constraint is applied by MoveLightCardTo, ensuring the light card is always flush to screens
diff --git a/Engine/Plugins/VirtualProduction/EpicStageApp/Source/EpicStageApp/Private/StageAppRouteHandler.cpp b/Engine/Plugins/VirtualProduction/EpicStageApp/Source/EpicStageApp/Private/StageAppRouteHandler.cpp
--- a/Engine/Plugins/VirtualProduction/EpicStageApp/Source/EpicStageApp/Private/StageAppRouteHandler.cpp
+++ b/Engine/Plugins/VirtualProduction/EpicStageApp/Source/EpicStageApp/Private/StageAppRouteHandler.cpp
@@ -723,6 +723,10 @@ void FStageAppRouteHandler::HandleWebSocketNDisplayPreviewActorDragBegin(const F
 
 		PerRendererData->PrimaryActor = FDisplayClusterWeakStageActorPtr(PrimaryActor);
 	}
+	else
+	{
+		PerRendererData->PrimaryActor.Reset();
+	}
 
 	if (!DragTimeoutTickerHandle.IsValid())
 	{
@@ -888,7 +892,7 @@ void FStageAppRouteHandler::HandleWebSocketNDisplayPreviewActorCreate(const FRem
 		FDisplayClusterLightCardEditorHelper& LightCardHelper = PerRendererData->GetLightCardHelper();
 
 		FSceneViewInitOptions ViewInitOptions;
-		if (PerRendererData->GetSceneViewInitOptions(ViewInitOptions, false))
+		if (PerRendererData->GetSceneViewInitOptions(ViewInitOptions))
 		{
 			const FSceneView View(ViewInitOptions);
 			LightCardHelper.MoveActorsToPixel({ NewActor }, PixelPos, View);
diff --git a/Engine/Plugins/VirtualProduction/EpicStageApp/Source/EpicStageApp/Public/StageAppVersion.h b/Engine/Plugins/VirtualProduction/EpicStageApp/Source/EpicStageApp/Public/StageAppVersion.h
--- a/Engine/Plugins/VirtualProduction/EpicStageApp/Source/EpicStageApp/Public/StageAppVersion.h
+++ b/Engine/Plugins/VirtualProduction/EpicStageApp/Source/EpicStageApp/Public/StageAppVersion.h
@@ -11,7 +11,7 @@ struct EPICSTAGEAPP_API FEpicStageAppAPIVersion
 	static const uint16 Major = 1;
 
 	/** The minor component of the current version. */
-	static const uint16 Minor = 7;
+	static const uint16 Minor = 8;
 	
 	/** The patch component of the current version. */
 	static const uint16 Patch = 0;
