[UnrealStage] Refactor beacon receiver to a generic base class

diff --git a/Engine/Plugins/DiscoveryBeaconReceiver/DiscoveryBeaconReceiver.uplugin b/Engine/Plugins/DiscoveryBeaconReceiver/DiscoveryBeaconReceiver.uplugin
new file mode 100644
--- /dev/null
+++ b/Engine/Plugins/DiscoveryBeaconReceiver/DiscoveryBeaconReceiver.uplugin
@@ -0,0 +1,25 @@
+{
+	"FileVersion": 1,
+	"Version": 1,
+	"VersionName": "1.0",
+	"FriendlyName": "Discovery Beacon Receiver",
+	"Description": "Listens for messages on a multicast socket and replies with information about the engine.",
+	"Category": "Messaging",
+	"CreatedBy": "Epic Games, Inc.",
+	"CreatedByURL": "https://epicgames.com",
+	"DocsURL": "",
+	"MarketplaceURL": "",
+	"SupportURL": "",
+	"CanContainContent": true,
+	"IsBetaVersion": false,
+	"IsExperimentalVersion": false,
+	"Installed": false,
+	"Modules": [
+		{
+			"Name": "DiscoveryBeaconReceiver",
+			"Type": "Runtime",
+			"LoadingPhase": "Default",
+			"PlatformAllowList": [ "Win64", "Linux", "Mac" ]
+		}
+	]
+}
\ No newline at end of file
diff --git a/Engine/Plugins/DiscoveryBeaconReceiver/Source/DiscoveryBeaconReceiver/DiscoveryBeaconReceiver.build.cs b/Engine/Plugins/DiscoveryBeaconReceiver/Source/DiscoveryBeaconReceiver/DiscoveryBeaconReceiver.build.cs
new file mode 100644
--- /dev/null
+++ b/Engine/Plugins/DiscoveryBeaconReceiver/Source/DiscoveryBeaconReceiver/DiscoveryBeaconReceiver.build.cs
@@ -0,0 +1,25 @@
+// Copyright Epic Games, Inc. All Rights Reserved.
+
+using UnrealBuildTool;
+
+public class DiscoveryBeaconReceiver : ModuleRules
+{
+	public DiscoveryBeaconReceiver(ReadOnlyTargetRules Target) : base(Target)
+	{
+		PublicDependencyModuleNames.AddRange(
+			new string[] {
+				"Core",
+				"CoreUObject",
+				"Serialization"
+			}
+		);
+
+		PrivateDependencyModuleNames.AddRange(
+			new string[] {
+				"Engine",
+				"Networking",
+				"Sockets"
+			}
+		);
+	}
+}
diff --git a/Engine/Plugins/DiscoveryBeaconReceiver/Source/DiscoveryBeaconReceiver/Private/DiscoveryBeaconReceiver.cpp b/Engine/Plugins/DiscoveryBeaconReceiver/Source/DiscoveryBeaconReceiver/Private/DiscoveryBeaconReceiver.cpp
new file mode 100644
--- /dev/null
+++ b/Engine/Plugins/DiscoveryBeaconReceiver/Source/DiscoveryBeaconReceiver/Private/DiscoveryBeaconReceiver.cpp
@@ -0,0 +1,194 @@
+// Copyright Epic Games, Inc. All Rights Reserved.
+
+#include "DiscoveryBeaconReceiver.h"
+
+#include "Common/UdpSocketBuilder.h"
+#include "DiscoveryBeaconReceiverModule.h"
+#include "MessageEndpoint.h"
+#include "MessageEndpointBuilder.h"
+#include "Misc/App.h"
+#include "Serialization/ArrayReader.h"
+#include "Serialization/ArrayWriter.h"
+
+FDiscoveryBeaconReceiver::FDiscoveryBeaconReceiver(const FString& InDescription, const TArray<uint8>& InProtocolIdentifier, uint8 InProtocolVersion)
+	: Description(InDescription), ProtocolIdentifier(InProtocolIdentifier), ProtocolVersion(InProtocolVersion), Guid(FGuid::NewGuid())
+{
+}
+
+void FDiscoveryBeaconReceiver::Startup()
+{
+	if (bIsRunning)
+	{
+		return;
+	}
+
+	FIPv4Address DiscoveryAddress;
+	if (!GetDiscoveryAddress(DiscoveryAddress))
+	{
+		return;
+	}
+
+	const int32 DiscoveryPort = GetDiscoveryPort();
+	if (DiscoveryPort < 0)
+	{
+		return;
+	}
+
+	Socket = FUdpSocketBuilder(Description)
+		.AsNonBlocking()
+		.AsReusable()
+		.BoundToPort(DiscoveryPort)
+		.WithMulticastLoopback()
+		.WithReceiveBufferSize(256);
+
+	if (!Socket)
+	{
+		UE_LOG(LogDiscoveryBeaconReceiver, Warning, TEXT("%s failed to create multicast socket"), *Description);
+		return;
+	}
+
+	bIsRunning = true;
+
+	// Listen for multicast packets on all interfaces.
+	// Note: If a new interface appears later on, the socket will need to join it in order to receive multicast packets from it.
+	{
+		TSharedRef<FInternetAddr> DiscoveryInternetAddr = FIPv4Endpoint(DiscoveryAddress, 0).ToInternetAddr(); // Note: Port will be ignored by JoinMulticastGroup
+
+		TArray<TSharedPtr<FInternetAddr>> LocapIps;
+		ISocketSubsystem::Get(PLATFORM_SOCKETSUBSYSTEM)->GetLocalAdapterAddresses(LocapIps);
+
+		for (TSharedPtr<FInternetAddr>& LocalIp : LocapIps)
+		{
+			if (!LocalIp.IsValid())
+			{
+				continue;
+			}
+
+			const bool bJoinedGroup = Socket->JoinMulticastGroup(*DiscoveryInternetAddr, *LocalIp);
+
+			if (bJoinedGroup)
+			{
+				UE_LOG(LogDiscoveryBeaconReceiver, Log,
+					TEXT("%s joined multicast group '%s' on detected local interface '%s'"),
+					*Description, *DiscoveryAddress.ToString(), *LocalIp->ToString(false));
+			}
+			else
+			{
+				UE_LOG(LogDiscoveryBeaconReceiver, Warning,
+					TEXT("%s failed to join multicast group '%s' on detected local interface '%s'"),
+					*Description, *DiscoveryAddress.ToString(), *LocalIp->ToString(false));
+			}
+		}
+	}
+
+	Thread.Reset(FRunnableThread::Create(this, *FString::Printf(TEXT("%sResponderThread"), *Description), 0, TPri_Normal, FPlatformAffinity::GetPoolThreadMask()));
+}
+
+void FDiscoveryBeaconReceiver::Shutdown()
+{
+	if (Thread)
+	{
+		Thread->Kill();
+	}
+
+	if (Socket)
+	{
+		ISocketSubsystem::Get()->DestroySocket(Socket);
+		Socket = nullptr;
+	}
+}
+
+bool FDiscoveryBeaconReceiver::Init()
+{
+	return true;
+}
+
+uint32 FDiscoveryBeaconReceiver::Run()
+{
+	const FTimespan WaitTime = FTimespan::FromSeconds(1.0);
+
+	while (bIsRunning)
+	{
+		if (!Socket)
+		{
+			return 0;
+		}
+
+		Socket->Wait(ESocketWaitConditions::WaitForRead, WaitTime);
+		ReceiveBeaconMessages();
+	}
+
+	return 0;
+}
+
+void FDiscoveryBeaconReceiver::Stop()
+{
+	bIsRunning = false;
+}
+
+void FDiscoveryBeaconReceiver::ReceiveBeaconMessages()
+{
+	if (!Socket)
+	{
+		return;
+	}
+
+	uint32 PendingDataSize;
+	if (Socket->HasPendingData(PendingDataSize) && PendingDataSize > 0)
+	{
+		FArrayReader MessageData = FArrayReader(true);
+		MessageData.SetNumUninitialized(PendingDataSize);
+
+		int32 NumRead;
+		TSharedRef<FInternetAddr> Source = ISocketSubsystem::Get()->CreateInternetAddr();
+		if (Socket->RecvFrom(MessageData.GetData(), MessageData.Num(), NumRead, *Source) && NumRead > 0)
+		{
+			HandleBeaconMessage(MessageData, Source);
+		}
+	}
+}
+
+void FDiscoveryBeaconReceiver::HandleBeaconMessage(FArrayReader& MessageData, TSharedRef<FInternetAddr> Source)
+{
+	// Check that the protocol identifier matches. If not, this message is probably unrelated.
+	for (const uint8 ExpectedByte : ProtocolIdentifier)
+	{
+		uint8 AppByte;
+		MessageData << AppByte;
+
+		if (AppByte != ExpectedByte)
+		{
+			return;
+		}
+	}
+
+	// We don't do anything with this yet, but could use it to ignore beacons from apps with incompatible protocols
+	uint8 BeaconProtocolVersion;
+	MessageData << BeaconProtocolVersion;
+
+	FArrayWriter Writer(true);
+
+#if !PLATFORM_LITTLE_ENDIAN
+	// App expects little-endian byte order
+	Writer.SetByteSwapping(true);
+#endif
+
+	{
+		Writer << (uint8&)ProtocolVersion;
+		Writer << Guid;
+	}
+
+	if (!MakeBeaconResponse(BeaconProtocolVersion, MessageData, Writer))
+	{
+		return;
+	}
+
+	// Message wasn't fully consumed.
+	if (!MessageData.AtEnd())
+	{
+		return;
+	}
+
+	int32 Sent;
+	Socket->SendTo(Writer.GetData(), Writer.Num(), Sent, *Source);
+}
\ No newline at end of file
diff --git a/Engine/Plugins/DiscoveryBeaconReceiver/Source/DiscoveryBeaconReceiver/Private/DiscoveryBeaconReceiverModule.cpp b/Engine/Plugins/DiscoveryBeaconReceiver/Source/DiscoveryBeaconReceiver/Private/DiscoveryBeaconReceiverModule.cpp
new file mode 100644
--- /dev/null
+++ b/Engine/Plugins/DiscoveryBeaconReceiver/Source/DiscoveryBeaconReceiver/Private/DiscoveryBeaconReceiverModule.cpp
@@ -0,0 +1,11 @@
+// Copyright Epic Games, Inc. All Rights Reserved.
+
+#include "DiscoveryBeaconReceiverModule.h"
+
+#define LOCTEXT_NAMESPACE "FDiscoveryBeaconReceiverModule"
+
+DEFINE_LOG_CATEGORY(LogDiscoveryBeaconReceiver);
+
+#undef LOCTEXT_NAMESPACE
+
+IMPLEMENT_MODULE(FDiscoveryBeaconReceiverModule, DiscoveryBeaconReceiver)
diff --git a/Engine/Plugins/DiscoveryBeaconReceiver/Source/DiscoveryBeaconReceiver/Public/DiscoveryBeaconReceiver.h b/Engine/Plugins/DiscoveryBeaconReceiver/Source/DiscoveryBeaconReceiver/Public/DiscoveryBeaconReceiver.h
new file mode 100644
--- /dev/null
+++ b/Engine/Plugins/DiscoveryBeaconReceiver/Source/DiscoveryBeaconReceiver/Public/DiscoveryBeaconReceiver.h
@@ -0,0 +1,101 @@
+// Copyright Epic Games, Inc. All Rights Reserved.
+
+#pragma once
+
+#include "CoreMinimal.h"
+#include "HAL/Runnable.h"
+#include "HAL/RunnableThread.h"
+
+class FArrayReader;
+class FBufferArchive;
+class FInternetAddr;
+class FSocket;
+struct FIPv4Address;
+typedef FBufferArchive FArrayWriter;
+
+/**
+ * Receives beacon messages on a multicast endpoint and replies with information to connect to the engine.
+ * This allows remote apps to detect compatible Unreal instances on the local network.
+ */
+class DISCOVERYBEACONRECEIVER_API FDiscoveryBeaconReceiver
+	: public FRunnable
+{
+public:
+	/**
+	 * Creates and initializes a new beacon receiver.
+	 * @param InDescription A short name for this beacon for debugging purposes.
+	 * @param InProtocolIdentifier A series of bytes that uniquely identify messages intended for this beacon receiver.
+	 * @param InProtocolVersion The protocol version number with which to reply to beacon messages.
+	 */
+	FDiscoveryBeaconReceiver(const FString& InDescription, const TArray<uint8>& InProtocolIdentifier, uint8 InProtocolVersion);
+
+	/** Destroy the beacon receiver. */
+	virtual ~FDiscoveryBeaconReceiver() = default;
+
+	/**
+	 * Open a socket and start a thread listening for beacon messages.
+	 * If called multiple times without shutting down, all but the first call will be ignored.
+	 */
+	virtual void Startup();
+
+	/** Close the socket and kill the listening thread. */
+	virtual void Shutdown();
+
+protected:
+	/**
+	 * Get the IP address on which to listen for beacon messages.
+	 * @param OutAddress The IP address.
+	 * @return false if no valid address could be found, else true.
+	 */
+	virtual bool GetDiscoveryAddress(FIPv4Address& OutAddress) const = 0;
+
+	/**
+	 * Get the port on which to listen for beacon messages.
+	 * @param OutAddress The port, or a negative value if no valid port could be found.
+	 */
+	virtual int32 GetDiscoveryPort() const = 0;
+
+	/**
+	 * Make a response to a beacon message.
+	 * @param BeaconProtocolVersion The protocol version reported by the beacon message.
+	 * @param InMessageData Array reader containing additional data after the standard beacon protocol version and identifier.
+	 *						All data in this reader must be consumed or the message will be ignored.
+	 * @param OutResponseData 
+	 */
+	virtual bool MakeBeaconResponse(uint8 BeaconProtocolVersion, FArrayReader& InMessageData, FArrayWriter& OutResponseData) const = 0;
+
+private:
+	//~ Begin FRunnable implementation
+	virtual bool Init() override;
+	virtual uint32 Run() override;
+	virtual void Stop() override;
+	//~ End FRunnable implementation
+
+	/** Called regularly to receive beacon messages. */
+	void ReceiveBeaconMessages();
+
+	/** Handle a message received on the beacon endpoint. */
+	void HandleBeaconMessage(FArrayReader& MessageData, TSharedRef<FInternetAddr> Source);
+
+private:
+	/** A short name for this beacon for debugging purposes. */
+	FString Description;
+
+	/** A series of bytes that uniquely identify messages intended for this beacon receiver. */
+	TArray<uint8> ProtocolIdentifier;
+
+	/** The protocol version number with which to reply to beacon messages. */
+	uint8 ProtocolVersion;
+
+	/** Identifier for this engine instance. Sent in responses to differentiate engine instances even if they're on the same machine. */
+	FGuid Guid;
+
+	/** Socket used to listen for and reply to beacon messages. */
+	FSocket* Socket = nullptr;
+
+	/** Flag indicating that the beacon thread should be running. */
+	TAtomic<bool> bIsRunning = false;
+
+	/** The thread to receive beacon messages on. */
+	TUniquePtr<FRunnableThread> Thread = nullptr;
+};
\ No newline at end of file
diff --git a/Engine/Plugins/DiscoveryBeaconReceiver/Source/DiscoveryBeaconReceiver/Public/DiscoveryBeaconReceiverModule.h b/Engine/Plugins/DiscoveryBeaconReceiver/Source/DiscoveryBeaconReceiver/Public/DiscoveryBeaconReceiverModule.h
new file mode 100644
--- /dev/null
+++ b/Engine/Plugins/DiscoveryBeaconReceiver/Source/DiscoveryBeaconReceiver/Public/DiscoveryBeaconReceiverModule.h
@@ -0,0 +1,11 @@
+// Copyright Epic Games, Inc. All Rights Reserved.
+
+#pragma once
+
+#include "Modules/ModuleManager.h"
+
+DISCOVERYBEACONRECEIVER_API DECLARE_LOG_CATEGORY_EXTERN(LogDiscoveryBeaconReceiver, Log, All);
+
+class FDiscoveryBeaconReceiverModule : public IModuleInterface
+{
+};
\ No newline at end of file
diff --git a/Engine/Plugins/VirtualProduction/EpicStageApp/EpicStageApp.uplugin b/Engine/Plugins/VirtualProduction/EpicStageApp/EpicStageApp.uplugin
--- a/Engine/Plugins/VirtualProduction/EpicStageApp/EpicStageApp.uplugin
+++ b/Engine/Plugins/VirtualProduction/EpicStageApp/EpicStageApp.uplugin
@@ -37,6 +37,10 @@
 		{
 			"Name": "nDisplayModularFeatures",
 			"Enabled": true
+		},
+		{
+			"Name": "DiscoveryBeaconReceiver",
+			"Enabled": true
 		}
 	]
 }
\ No newline at end of file
diff --git a/Engine/Plugins/VirtualProduction/EpicStageApp/Source/EpicStageApp/EpicStageApp.build.cs b/Engine/Plugins/VirtualProduction/EpicStageApp/Source/EpicStageApp/EpicStageApp.build.cs
--- a/Engine/Plugins/VirtualProduction/EpicStageApp/Source/EpicStageApp/EpicStageApp.build.cs
+++ b/Engine/Plugins/VirtualProduction/EpicStageApp/Source/EpicStageApp/EpicStageApp.build.cs
@@ -17,6 +17,7 @@ public class EpicStageApp : ModuleRules
 		PrivateDependencyModuleNames.AddRange(
 			new string[] {
 				"DeveloperSettings",
+				"DiscoveryBeaconReceiver",
 				"DisplayCluster",
 				"DisplayClusterLightCardEditorShaders",
 				"DisplayClusterLightCardExtender",
diff --git a/Engine/Plugins/VirtualProduction/EpicStageApp/Source/EpicStageApp/Private/EpicStageApp.cpp b/Engine/Plugins/VirtualProduction/EpicStageApp/Source/EpicStageApp/Private/EpicStageApp.cpp
--- a/Engine/Plugins/VirtualProduction/EpicStageApp/Source/EpicStageApp/Private/EpicStageApp.cpp
+++ b/Engine/Plugins/VirtualProduction/EpicStageApp/Source/EpicStageApp/Private/EpicStageApp.cpp
@@ -12,14 +12,14 @@ void FEpicStageAppModule::StartupModule()
 	{
 		if (WebRemoteControl.IsWebSocketServerRunning())
 		{
-			StartupBeaconReceiver();
+			StageAppBeaconReceiver.Startup();
 		}
 		else
 		{
 			// Wait for the WebSocket server to start
 			WebSocketServerStartedDelegate = WebRemoteControl.OnWebSocketServerStarted().AddLambda([this](uint32)
 			{
-				StartupBeaconReceiver();
+				StageAppBeaconReceiver.Startup();
 			});
 		}
 	}
@@ -37,23 +37,5 @@ void FEpicStageAppModule::ShutdownModule()
 		}
 	}
 
-	ShutdownBeaconReceiver();
-}
-
-void FEpicStageAppModule::StartupBeaconReceiver()
-{
-	// Make sure it's not already running
-	ShutdownBeaconReceiver();
-
-	StageAppBeaconReceiver.Startup();
-	bIsBeaconReceiverRunning = true;
-}
-
-void FEpicStageAppModule::ShutdownBeaconReceiver()
-{
-	if (bIsBeaconReceiverRunning)
-	{
-		StageAppBeaconReceiver.Shutdown();
-		bIsBeaconReceiverRunning = false;
-	}
+	StageAppBeaconReceiver.Shutdown();
 }
diff --git a/Engine/Plugins/VirtualProduction/EpicStageApp/Source/EpicStageApp/Private/EpicStageApp.h b/Engine/Plugins/VirtualProduction/EpicStageApp/Source/EpicStageApp/Private/EpicStageApp.h
--- a/Engine/Plugins/VirtualProduction/EpicStageApp/Source/EpicStageApp/Private/EpicStageApp.h
+++ b/Engine/Plugins/VirtualProduction/EpicStageApp/Source/EpicStageApp/Private/EpicStageApp.h
@@ -19,20 +19,10 @@ class FEpicStageAppModule : public IModuleInterface
 	/** Handler for WebSocket routes. */
 	FStageAppRouteHandler RouteHandler;
 
+private:
 	/** Receiver that responds to beacon messages from the stage app. */
 	FStageAppBeaconReceiver StageAppBeaconReceiver;
 
-private:
-	/** Start the beacon receiver so the engine can be seen by the app. */
-	void StartupBeaconReceiver();
-
-	/** Shut down the beacon receiver so the engine can no longer be seen by the app. */
-	void ShutdownBeaconReceiver();
-
-private:
-	/** Whether the beacon receiver has been started and has not been stopped. */
-	bool bIsBeaconReceiverRunning = false;
-
 	/** Delegate for when the websocket server starts. */
 	FDelegateHandle WebSocketServerStartedDelegate;
 };
diff --git a/Engine/Plugins/VirtualProduction/EpicStageApp/Source/EpicStageApp/Private/StageAppBeaconReceiver.cpp b/Engine/Plugins/VirtualProduction/EpicStageApp/Source/EpicStageApp/Private/StageAppBeaconReceiver.cpp
--- a/Engine/Plugins/VirtualProduction/EpicStageApp/Source/EpicStageApp/Private/StageAppBeaconReceiver.cpp
+++ b/Engine/Plugins/VirtualProduction/EpicStageApp/Source/EpicStageApp/Private/StageAppBeaconReceiver.cpp
@@ -20,176 +20,51 @@ namespace StageAppBeaconReceiverConstants
 }
 
 FStageAppBeaconReceiver::FStageAppBeaconReceiver()
-	: Guid(FGuid::NewGuid())
+	: FDiscoveryBeaconReceiver(
+		TEXT("StageAppBeaconResponder"),
+		StageAppBeaconReceiverConstants::ProtocolIdentifier,
+		StageAppBeaconReceiverConstants::ProtocolVersion
+	)
 {
 }
 
 void FStageAppBeaconReceiver::Startup()
 {
-	const UStageAppSettings& Settings = *GetDefault<UStageAppSettings>();
-	FIPv4Address DiscoveryAddress;
-	if (!FIPv4Address::Parse(Settings.DiscoveryEndpoint, DiscoveryAddress))
-	{
-		UE_LOG(LogRemoteControl, Error, TEXT("Failed to parse Stage App discovery endpoint address \"%s\""), *Settings.DiscoveryEndpoint);
-		return;
-	}
-
 	WebsocketPort = GetDefault<URemoteControlSettings>()->RemoteControlWebSocketServerPort;
 
-	Socket = FUdpSocketBuilder(TEXT("StageAppBeaconResponder"))
-		.AsNonBlocking()
-		.AsReusable()
-		.BoundToPort(Settings.DiscoveryPort)
-		.WithMulticastLoopback()
-		.WithReceiveBufferSize(256);
-
-	if (!Socket)
-	{
-		UE_LOG(LogRemoteControl, Warning, TEXT("StageAppBeaconReceiver failed to create multicast socket"));
-		return;
-	}
-
-	// Listen for multicast packets on all interfaces.
-	// Note: If a new interface appears later on, the socket will need to join it in order to receive multicast packets from it.
-	{
-		TSharedRef<FInternetAddr> DiscoveryInternetAddr = FIPv4Endpoint(DiscoveryAddress, 0).ToInternetAddr(); // Note: Port will be ignored by JoinMulticastGroup
-
-		TArray<TSharedPtr<FInternetAddr>> LocapIps;
-		ISocketSubsystem::Get(PLATFORM_SOCKETSUBSYSTEM)->GetLocalAdapterAddresses(LocapIps);
-
-		for (TSharedPtr<FInternetAddr>& LocalIp : LocapIps)
-		{
-			if (!LocalIp.IsValid())
-			{
-				continue;
-			}
-			
-			const bool bJoinedGroup = Socket->JoinMulticastGroup(*DiscoveryInternetAddr, *LocalIp);
-
-			if (bJoinedGroup)
-			{
-				UE_LOG(LogRemoteControl, Log,
-					TEXT("StageAppBeaconReceiver joined multicast group '%s' on detected local interface '%s'"),
-					*DiscoveryAddress.ToString(), *LocalIp->ToString(false));
-			}
-			else
-			{
-				UE_LOG(LogRemoteControl, Warning,
-					TEXT("StageAppBeaconReceiver failed to join multicast group '%s' on detected local interface '%s'"),
-					*DiscoveryAddress.ToString(), *LocalIp->ToString(false));
-			}
-		}
-	}
-
-	Thread.Reset(FRunnableThread::Create(this, TEXT("StageAppBeaconResponderThread"), 0, TPri_Normal, FPlatformAffinity::GetPoolThreadMask()));
+	FDiscoveryBeaconReceiver::Startup();
 }
 
-void FStageAppBeaconReceiver::Shutdown()
+bool FStageAppBeaconReceiver::GetDiscoveryAddress(FIPv4Address& OutAddress) const
 {
-	if (Thread)
-	{
-		Thread->Kill();
-		Thread = nullptr;
-	}
-
-	if (Socket)
+	const UStageAppSettings& Settings = *GetDefault<UStageAppSettings>();
+	FIPv4Address DiscoveryAddress;
+	if (!FIPv4Address::Parse(Settings.DiscoveryEndpoint, DiscoveryAddress))
 	{
-		ISocketSubsystem::Get()->DestroySocket(Socket);
+		UE_LOG(LogRemoteControl, Error, TEXT("Failed to parse Stage App discovery endpoint address \"%s\""), *Settings.DiscoveryEndpoint);
+		return false;
 	}
-}
 
-bool FStageAppBeaconReceiver::Init()
-{
+	OutAddress = DiscoveryAddress;
 	return true;
 }
 
-uint32 FStageAppBeaconReceiver::Run()
+int32 FStageAppBeaconReceiver::GetDiscoveryPort() const
 {
 	const UStageAppSettings& Settings = *GetDefault<UStageAppSettings>();
-	const FTimespan WaitTime = FTimespan::FromSeconds(Settings.DiscoverySocketWaitTime);
-
-	while (!bStopping)
-	{
-		if (!Socket)
-		{
-			return 0;
-		}
-
-		Socket->Wait(ESocketWaitConditions::WaitForRead, WaitTime);
-		ReceiveBeaconMessages();
-	}
-
-	return 0;
+	return Settings.DiscoveryPort;
 }
 
-void FStageAppBeaconReceiver::Stop()
+bool FStageAppBeaconReceiver::MakeBeaconResponse(uint8 BeaconProtocolVersion, FArrayReader& InMessageData, FArrayWriter& OutResponseData) const
 {
-	bStopping = true;
-}
-
-void FStageAppBeaconReceiver::ReceiveBeaconMessages()
-{
-	if (!Socket)
-	{
-		return;
-	}
-
-	uint32 PendingDataSize;
-	if (Socket->HasPendingData(PendingDataSize) && PendingDataSize > 0)
-	{
-		FArrayReader MessageData = FArrayReader(true);
-		MessageData.SetNumUninitialized(PendingDataSize);
-
-		int32 NumRead;
-		TSharedRef<FInternetAddr> Source = ISocketSubsystem::Get()->CreateInternetAddr();
-		if (Socket->RecvFrom(MessageData.GetData(), MessageData.Num(), NumRead, *Source) && NumRead > 0)
-		{
-			HandleBeaconMessage(MessageData, Source);
-		}
-	}
-}
-
-void FStageAppBeaconReceiver::HandleBeaconMessage(FArrayReader& MessageData, TSharedRef<FInternetAddr> Source)
-{
-	if (MessageData.Num() != 5)
-	{
-		return;
-	}
-
-	// Check that the protocol identifier matches. If not, this message is probably unrelated.
-	for (const uint8 ExpectedByte : StageAppBeaconReceiverConstants::ProtocolIdentifier)
-	{
-		uint8 AppByte;
-		MessageData << AppByte;
-
-		if (AppByte != ExpectedByte)
-		{
-			return;
-		}
-	}
-
-	// We don't do anything with this yet, but could use it to ignore beacons from apps with incompatible protocols
-	uint8 AppProtocolVersion;
-	MessageData << AppProtocolVersion;
-
 	const FString FriendlyName = GetFriendlyName();
 
-	FArrayWriter Writer(true);
-
-#if !PLATFORM_LITTLE_ENDIAN
-	// App expects little-endian byte order
-	Writer.SetByteSwapping(true);
-#endif
-
 	{
-		Writer << (uint8&)StageAppBeaconReceiverConstants::ProtocolVersion;
-		Writer << Guid;
-		Writer << (uint32&)WebsocketPort;
-		Writer << (FString&)FriendlyName;
+		OutResponseData << (uint32&)WebsocketPort;
+		OutResponseData << (FString&)FriendlyName;
 	}
 
-	int32 Sent;
-	Socket->SendTo(Writer.GetData(), Writer.Num(), Sent, *Source);
+	return true;
 }
 
 FString FStageAppBeaconReceiver::GetFriendlyName() const
diff --git a/Engine/Plugins/VirtualProduction/EpicStageApp/Source/EpicStageApp/Private/StageAppBeaconReceiver.h b/Engine/Plugins/VirtualProduction/EpicStageApp/Source/EpicStageApp/Private/StageAppBeaconReceiver.h
--- a/Engine/Plugins/VirtualProduction/EpicStageApp/Source/EpicStageApp/Private/StageAppBeaconReceiver.h
+++ b/Engine/Plugins/VirtualProduction/EpicStageApp/Source/EpicStageApp/Private/StageAppBeaconReceiver.h
@@ -2,91 +2,33 @@
 
 #pragma once
 
-#include "CoreMinimal.h"
-#include "HAL/Runnable.h"
-#include "HAL/RunnableThread.h"
-#include "StageAppBeaconReceiver.generated.h"
-
-class FArrayReader;
-class FInternetAddr;
-class FSocket;
-
-USTRUCT()
-struct FRemoteControl_StageAppBeaconMessage
-{
-	GENERATED_BODY()
-};
-
-USTRUCT()
-struct FRemoteControl_StageAppConnectionInfoMessage
-{
-	GENERATED_BODY()
-
-	/** Version of the stage app beacon protocol to use. */
-	uint8 Protocol;
-
-	/** ID of the sender. */
-	FGuid SenderId;
-
-	/** The websocket port the app should connect to. */
-	uint32 WebsocketPort;
-
-	/** The websocket address the app should connect to. */
-	FString WebsocketAddress;
-
-	/** The user-friendly name of the connection to show to the user. */
-	FString FriendlyName;
-};
+#include "DiscoveryBeaconReceiver.h"
 
 /**
  * Receives beacon messages from the Unreal Stage App and replies with connection information.
  * This allows the app to detect compatible Unreal instances on the local network and list them for the user.
  */
 class FStageAppBeaconReceiver
-	: public FRunnable
+	: public FDiscoveryBeaconReceiver
 {
 public:
-	/** Creates and initializes a new beacon receiver. */
 	FStageAppBeaconReceiver();
-
-	/** Destroy the beacon receiver. */
 	virtual ~FStageAppBeaconReceiver() {};
 
-	/** Open a socket and start a thread listening for beacon messages. */
-	void Startup();
-
-	/** Close the socket and kill the listening thread. */
-	void Shutdown();
+	//~ Begin FDiscoveryBeaconReceiver implementation
+public:
+	virtual void Startup() override;
 
 protected:
-	//~ Begin FRunnable implementation
-	virtual bool Init() override;
-	virtual uint32 Run() override;
-	virtual void Stop() override;
-	//~ End FRunnable implementation
+	virtual bool GetDiscoveryAddress(FIPv4Address& OutAddress) const override;
+	virtual int32 GetDiscoveryPort() const override;
+	virtual bool MakeBeaconResponse(uint8 BeaconProtocolVersion, FArrayReader& InMessageData, FArrayWriter& OutResponseData) const override;
+	//~ End FDiscoveryBeaconReceiver implementation
 
 private:
-	/** Called regularly to receive beacon messages. */
-	void ReceiveBeaconMessages();
-
-	/** Handle a message received on the beacon endpoint. */
-	void HandleBeaconMessage(FArrayReader& MessageData, TSharedRef<FInternetAddr> Source);
-
 	/** Get the name to report to apps searching for the engine. */
 	FString GetFriendlyName() const;
 
-	/** Identifier for this engine instance. Used by the app to differentiate instances even if they're on the same machine. */
-	FGuid Guid;
-
-	/** The websocket port to reply with. */
+	/** The websocket port with which to reply. */
 	uint32 WebsocketPort;
-
-	/** Socket used to listen for and reply to beacon messages from the app. */
-	FSocket* Socket = nullptr;
-
-	/** Flag indicating that the thread is stopping. */
-	TAtomic<bool> bStopping;
-
-	/** The thread to receive beacon messages on. */
-	TUniquePtr<FRunnableThread> Thread = nullptr;
 };
\ No newline at end of file
diff --git a/Engine/Plugins/VirtualProduction/EpicStageApp/Source/EpicStageApp/Public/StageAppSettings.h b/Engine/Plugins/VirtualProduction/EpicStageApp/Source/EpicStageApp/Public/StageAppSettings.h
--- a/Engine/Plugins/VirtualProduction/EpicStageApp/Source/EpicStageApp/Public/StageAppSettings.h
+++ b/Engine/Plugins/VirtualProduction/EpicStageApp/Source/EpicStageApp/Public/StageAppSettings.h
@@ -52,6 +52,7 @@ public:
 	/**
 	 * How many seconds to wait for app discovery messages.
 	 */
+	UE_DEPRECATED(5.4, "This setting is no longer used.")
 	UPROPERTY(config, EditAnywhere, Category = "Epic Stage App")
 	float DiscoverySocketWaitTime = 10.f;
 };
