[UnrealStage] Fix preview renders not being forced on when Unreal Stage is connected

diff --git a/Engine/Plugins/Runtime/nDisplay/Source/DisplayCluster/Private/Game/EngineClasses/Scene/DisplayClusterRootActor.cpp b/Engine/Plugins/Runtime/nDisplay/Source/DisplayCluster/Private/Game/EngineClasses/Scene/DisplayClusterRootActor.cpp
--- a/Engine/Plugins/Runtime/nDisplay/Source/DisplayCluster/Private/Game/EngineClasses/Scene/DisplayClusterRootActor.cpp
+++ b/Engine/Plugins/Runtime/nDisplay/Source/DisplayCluster/Private/Game/EngineClasses/Scene/DisplayClusterRootActor.cpp
@@ -1125,6 +1125,15 @@ void ADisplayClusterRootActor::Tick(float DeltaSeconds)
 		NewPreviewSettings.bPreviewEnable = false;
 	}
 
+#if WITH_EDITOR
+	if (!PreviewEnableOverriders.IsEmpty())
+	{
+		// Preview rendering is overridden, so force it on
+		bEnablePreviewInScene = true;
+		NewPreviewSettings.bPreviewEnable = true;
+	}
+#endif
+
 	// Update entire cluster preview rendering
 	if (IDisplayClusterViewportManager* ViewportManager = bEnablePreviewInScene ? GetOrCreateViewportManager() : GetViewportManager())
 	{
diff --git a/Engine/Plugins/Runtime/nDisplay/Source/DisplayCluster/Private/Game/EngineClasses/Scene/DisplayClusterRootActorEditor.cpp b/Engine/Plugins/Runtime/nDisplay/Source/DisplayCluster/Private/Game/EngineClasses/Scene/DisplayClusterRootActorEditor.cpp
--- a/Engine/Plugins/Runtime/nDisplay/Source/DisplayCluster/Private/Game/EngineClasses/Scene/DisplayClusterRootActorEditor.cpp
+++ b/Engine/Plugins/Runtime/nDisplay/Source/DisplayCluster/Private/Game/EngineClasses/Scene/DisplayClusterRootActorEditor.cpp
@@ -346,6 +346,18 @@ void ADisplayClusterRootActor::RerunConstructionScripts_Editor()
 	UpdateInnerFrustumPriority();
 }
 
+void ADisplayClusterRootActor::AddPreviewEnableOverride(const uint8* Object)
+{
+	check(Object);
+	PreviewEnableOverriders.Add(Object);
+}
+
+void ADisplayClusterRootActor::RemovePreviewEnableOverride(const uint8* Object)
+{
+	check(Object);
+	PreviewEnableOverriders.Remove(Object);
+}
+
 void ADisplayClusterRootActor::UpdateInnerFrustumPriority()
 {
 	if (InnerFrustumPriority.Num() == 0)
diff --git a/Engine/Plugins/Runtime/nDisplay/Source/DisplayCluster/Public/DisplayClusterRootActor.h b/Engine/Plugins/Runtime/nDisplay/Source/DisplayCluster/Public/DisplayClusterRootActor.h
--- a/Engine/Plugins/Runtime/nDisplay/Source/DisplayCluster/Public/DisplayClusterRootActor.h
+++ b/Engine/Plugins/Runtime/nDisplay/Source/DisplayCluster/Public/DisplayClusterRootActor.h
@@ -623,15 +623,13 @@ public:
 	bool ShouldThisFrameOutputPreviewToPostProcessRenderTarget() const { return false; }
 
 	/** Force preview rendering to be enabled regardless of the user's setting until a matching RemovePreviewEnableOverride call is made. */
-	UE_DEPRECATED(5.4, "This function has been deprecated.")
-	void AddPreviewEnableOverride(const uint8* Object) { };
+	void AddPreviewEnableOverride(const uint8* Object);
 
 	/**
 	 * Stop forcing preview rendering to be enabled for this caller. If other objects have called AddPreviewEnableOverride, it will remain
 	 * forced until they have also removed their overrides.
 	 */
-	UE_DEPRECATED(5.4, "This function has been deprecated.")
-	void RemovePreviewEnableOverride(const uint8* Object) { };
+	void RemovePreviewEnableOverride(const uint8* Object);
 
 	UE_DEPRECATED(5.4, "This function has been deprecated.")
 	float GetPreviewRenderTargetRatioMult() const
@@ -684,6 +682,9 @@ private:
 
 	bool bIsSelectedInEditor = false;
 
+	/* Addresses of callers to AddPreviewEnableOverride that haven't removed their overrides yet. */
+	TSet<const uint8*> PreviewEnableOverriders;
+
 	TWeakPtr<IDisplayClusterConfiguratorBlueprintEditor> ToolkitPtr;
 
 	// UE_DEPRECATED 5.4
diff --git a/Engine/Plugins/VirtualProduction/EpicStageApp/Source/EpicStageApp/Private/StageAppRouteHandler.cpp b/Engine/Plugins/VirtualProduction/EpicStageApp/Source/EpicStageApp/Private/StageAppRouteHandler.cpp
--- a/Engine/Plugins/VirtualProduction/EpicStageApp/Source/EpicStageApp/Private/StageAppRouteHandler.cpp
+++ b/Engine/Plugins/VirtualProduction/EpicStageApp/Source/EpicStageApp/Private/StageAppRouteHandler.cpp
@@ -414,6 +414,7 @@ void FStageAppRouteHandler::HandleWebSocketNDisplayPreviewRendererCreate(const F
 	if (!Body.RootActorPath.IsEmpty())
 	{
 		PreviewModule.SetRendererRootActorPath(RendererId, Body.RootActorPath, true);
+		BeginForceRootActorPreview(RendererId);
 	}
 
 	ChangePreviewRendererSettings(WebSocketMessage.ClientId, RendererId, Body.Settings);
@@ -435,7 +436,11 @@ void FStageAppRouteHandler::HandleWebSocketNDisplayPreviewRendererSetRoot(const
 		return;
 	}
 
+	EndForceRootActorPreview(Body.RendererId);
+
 	IDisplayClusterScenePreview::Get().SetRendererRootActorPath(Body.RendererId, Body.RootActorPath, true);
+
+	BeginForceRootActorPreview(Body.RendererId);
 }
 
 void FStageAppRouteHandler::HandleWebSocketNDisplayPreviewRendererConfigure(const FRemoteControlWebSocketMessage& WebSocketMessage)
@@ -474,6 +479,7 @@ void FStageAppRouteHandler::HandleWebSocketNDisplayPreviewRendererDestroy(const
 		// Check that this client created the renderer
 		if (PerRendererDataMap->Remove(Body.RendererId) > 0)
 		{
+			EndForceRootActorPreview(Body.RendererId);
 			IDisplayClusterScenePreview::Get().DestroyRenderer(Body.RendererId);
 		}
 	}
@@ -1039,6 +1045,8 @@ void FStageAppRouteHandler::HandleClientDisconnected(FGuid ClientId)
 		for (TClientIdToPerRendererDataMap::ElementType& PerRendererDataPair : *PerRendererDataMap)
 		{
 			EndActorDrag(PerRendererDataPair.Value, ClientId, PerRendererDataPair.Key, true);
+
+			EndForceRootActorPreview(PerRendererDataPair.Key);
 			PreviewModule.DestroyRenderer(PerRendererDataPair.Key);
 		}
 
@@ -1189,4 +1197,25 @@ void FStageAppRouteHandler::EndActorDrag(FPerRendererData& PerRendererData, cons
 	}
 }
 
+void FStageAppRouteHandler::BeginForceRootActorPreview(int32 RendererId)
+{
+#if WITH_EDITOR
+	if (ADisplayClusterRootActor* RootActor = IDisplayClusterScenePreview::Get().GetRendererRootActor(RendererId))
+	{
+		// Force the new root actor to render previews so the app always gets a render of it
+		RootActor->AddPreviewEnableOverride(reinterpret_cast<uint8*>(this));
+	}
+#endif
+}
+
+void FStageAppRouteHandler::EndForceRootActorPreview(int32 RendererId)
+{
+#if WITH_EDITOR
+	if (ADisplayClusterRootActor* RootActor = IDisplayClusterScenePreview::Get().GetRendererRootActor(RendererId))
+	{
+		RootActor->RemovePreviewEnableOverride(reinterpret_cast<uint8*>(this));
+	};
+#endif
+}
+
 #undef LOCTEXT_NAMESPACE
diff --git a/Engine/Plugins/VirtualProduction/EpicStageApp/Source/EpicStageApp/Private/StageAppRouteHandler.h b/Engine/Plugins/VirtualProduction/EpicStageApp/Source/EpicStageApp/Private/StageAppRouteHandler.h
--- a/Engine/Plugins/VirtualProduction/EpicStageApp/Source/EpicStageApp/Private/StageAppRouteHandler.h
+++ b/Engine/Plugins/VirtualProduction/EpicStageApp/Source/EpicStageApp/Private/StageAppRouteHandler.h
@@ -164,6 +164,15 @@ private:
 	/** End a client's actor drag operation. */
 	void EndActorDrag(FPerRendererData& PreviewData, const FGuid& ClientId, int32 RendererId, bool bEndedByClient);
 
+	/**
+	 * If the preview renderer with the given ID has a root actor, force it to render previews until EndForceRootActorPreview is called
+	 * with the same ID.
+	 */
+	void BeginForceRootActorPreview(int32 RendererId);
+
+	/** If the preview renderer with the given ID has a root actor, stop forcing it to render previews. */
+	void EndForceRootActorPreview(int32 RendererId);
+
 private:
 	/** The module for the WebSocket server for which this is handling routes. */
 	IWebRemoteControlModule* RemoteControlModule = nullptr;
